Fixing copyright in Lumin Plugins

diff --git a/Engine/Plugins/Lumin/MagicLeap/Source/MLSDK/MLSDK.Build.cs b/Engine/Plugins/Lumin/MagicLeap/Source/MLSDK/MLSDK.Build.cs
--- a/Engine/Plugins/Lumin/MagicLeap/Source/MLSDK/MLSDK.Build.cs
+++ b/Engine/Plugins/Lumin/MagicLeap/Source/MLSDK/MLSDK.Build.cs
@@ -1,36 +1,4 @@
-// %BANNER_BEGIN%
-// ---------------------------------------------------------------------
-// %COPYRIGHT_BEGIN%
-//
-// Copyright (c) 2017 Magic Leap, Inc. (COMPANY) All Rights Reserved.
-// Magic Leap, Inc. Confidential and Proprietary
-//
-// NOTICE: All information contained herein is, and remains the property
-// of COMPANY. The intellectual and technical concepts contained herein
-// are proprietary to COMPANY and may be covered by U.S. and Foreign
-// Patents, patents in process, and are protected by trade secret or
-// copyright law. Dissemination of this information or reproduction of
-// this material is strictly forbidden unless prior written permission is
-// obtained from COMPANY. Access to the source code contained herein is
-// hereby forbidden to anyone except current COMPANY employees, managers
-// or contractors who have executed Confidentiality and Non-disclosure
-// agreements explicitly covering such access.
-//
-// The copyright notice above does not evidence any actual or intended
-// publication or disclosure of this source code, which includes
-// information that is confidential and/or proprietary, and is a trade
-// secret, of COMPANY. ANY REPRODUCTION, MODIFICATION, DISTRIBUTION,
-// PUBLIC PERFORMANCE, OR PUBLIC DISPLAY OF OR THROUGH USE OF THIS
-// SOURCE CODE WITHOUT THE EXPRESS WRITTEN CONSENT OF COMPANY IS
-// STRICTLY PROHIBITED, AND IN VIOLATION OF APPLICABLE LAWS AND
-// INTERNATIONAL TREATIES. THE RECEIPT OR POSSESSION OF THIS SOURCE
-// CODE AND/OR RELATED INFORMATION DOES NOT CONVEY OR IMPLY ANY RIGHTS
-// TO REPRODUCE, DISCLOSE OR DISTRIBUTE ITS CONTENTS, OR TO MANUFACTURE,
-// USE, OR SELL ANYTHING THAT IT MAY DESCRIBE, IN WHOLE OR IN PART.
-//
-// %COPYRIGHT_END%
-// --------------------------------------------------------------------
-// %BANNER_END%
+// Copyright 1998-2018 Epic Games, Inc. All Rights Reserved.
 
 using UnrealBuildTool;
 using System.IO;
diff --git a/Engine/Plugins/Lumin/MagicLeap/Source/MLSDK/Public/MagicLeapMath.h b/Engine/Plugins/Lumin/MagicLeap/Source/MLSDK/Public/MagicLeapMath.h
--- a/Engine/Plugins/Lumin/MagicLeap/Source/MLSDK/Public/MagicLeapMath.h
+++ b/Engine/Plugins/Lumin/MagicLeap/Source/MLSDK/Public/MagicLeapMath.h
@@ -1,36 +1,4 @@
-// %BANNER_BEGIN%
-// ---------------------------------------------------------------------
-// %COPYRIGHT_BEGIN%
-//
-// Copyright (c) 2017 Magic Leap, Inc. (COMPANY) All Rights Reserved.
-// Magic Leap, Inc. Confidential and Proprietary
-//
-// NOTICE: All information contained herein is, and remains the property
-// of COMPANY. The intellectual and technical concepts contained herein
-// are proprietary to COMPANY and may be covered by U.S. and Foreign
-// Patents, patents in process, and are protected by trade secret or
-// copyright law. Dissemination of this information or reproduction of
-// this material is strictly forbidden unless prior written permission is
-// obtained from COMPANY. Access to the source code contained herein is
-// hereby forbidden to anyone except current COMPANY employees, managers
-// or contractors who have executed Confidentiality and Non-disclosure
-// agreements explicitly covering such access.
-//
-// The copyright notice above does not evidence any actual or intended
-// publication or disclosure of this source code, which includes
-// information that is confidential and/or proprietary, and is a trade
-// secret, of COMPANY. ANY REPRODUCTION, MODIFICATION, DISTRIBUTION,
-// PUBLIC PERFORMANCE, OR PUBLIC DISPLAY OF OR THROUGH USE OF THIS
-// SOURCE CODE WITHOUT THE EXPRESS WRITTEN CONSENT OF COMPANY IS
-// STRICTLY PROHIBITED, AND IN VIOLATION OF APPLICABLE LAWS AND
-// INTERNATIONAL TREATIES. THE RECEIPT OR POSSESSION OF THIS SOURCE
-// CODE AND/OR RELATED INFORMATION DOES NOT CONVEY OR IMPLY ANY RIGHTS
-// TO REPRODUCE, DISCLOSE OR DISTRIBUTE ITS CONTENTS, OR TO MANUFACTURE,
-// USE, OR SELL ANYTHING THAT IT MAY DESCRIBE, IN WHOLE OR IN PART.
-//
-// %COPYRIGHT_END%
-// --------------------------------------------------------------------
-// %BANNER_END%
+// Copyright 1998-2018 Epic Games, Inc. All Rights Reserved.
 
 #pragma once
 
diff --git a/Engine/Plugins/Lumin/MagicLeap/Source/MLSDK/Public/MagicLeapPluginUtil.h b/Engine/Plugins/Lumin/MagicLeap/Source/MLSDK/Public/MagicLeapPluginUtil.h
--- a/Engine/Plugins/Lumin/MagicLeap/Source/MLSDK/Public/MagicLeapPluginUtil.h
+++ b/Engine/Plugins/Lumin/MagicLeap/Source/MLSDK/Public/MagicLeapPluginUtil.h
@@ -1,36 +1,4 @@
-// %BANNER_BEGIN%
-// ---------------------------------------------------------------------
-// %COPYRIGHT_BEGIN%
-//
-// Copyright (c) 2017 Magic Leap, Inc. (COMPANY) All Rights Reserved.
-// Magic Leap, Inc. Confidential and Proprietary
-//
-// NOTICE: All information contained herein is, and remains the property
-// of COMPANY. The intellectual and technical concepts contained herein
-// are proprietary to COMPANY and may be covered by U.S. and Foreign
-// Patents, patents in process, and are protected by trade secret or
-// copyright law. Dissemination of this information or reproduction of
-// this material is strictly forbidden unless prior written permission is
-// obtained from COMPANY. Access to the source code contained herein is
-// hereby forbidden to anyone except current COMPANY employees, managers
-// or contractors who have executed Confidentiality and Non-disclosure
-// agreements explicitly covering such access.
-//
-// The copyright notice above does not evidence any actual or intended
-// publication or disclosure of this source code, which includes
-// information that is confidential and/or proprietary, and is a trade
-// secret, of COMPANY. ANY REPRODUCTION, MODIFICATION, DISTRIBUTION,
-// PUBLIC PERFORMANCE, OR PUBLIC DISPLAY OF OR THROUGH USE OF THIS
-// SOURCE CODE WITHOUT THE EXPRESS WRITTEN CONSENT OF COMPANY IS
-// STRICTLY PROHIBITED, AND IN VIOLATION OF APPLICABLE LAWS AND
-// INTERNATIONAL TREATIES. THE RECEIPT OR POSSESSION OF THIS SOURCE
-// CODE AND/OR RELATED INFORMATION DOES NOT CONVEY OR IMPLY ANY RIGHTS
-// TO REPRODUCE, DISCLOSE OR DISTRIBUTE ITS CONTENTS, OR TO MANUFACTURE,
-// USE, OR SELL ANYTHING THAT IT MAY DESCRIBE, IN WHOLE OR IN PART.
-//
-// %COPYRIGHT_END%
-// --------------------------------------------------------------------
-// %BANNER_END%
+// Copyright 1998-2018 Epic Games, Inc. All Rights Reserved.
 
 #pragma once
 #include "CoreMinimal.h"
diff --git a/Engine/Plugins/Lumin/MagicLeap/Source/MagicLeap/Classes/CameraCaptureComponent.h b/Engine/Plugins/Lumin/MagicLeap/Source/MagicLeap/Classes/CameraCaptureComponent.h
--- a/Engine/Plugins/Lumin/MagicLeap/Source/MagicLeap/Classes/CameraCaptureComponent.h
+++ b/Engine/Plugins/Lumin/MagicLeap/Source/MagicLeap/Classes/CameraCaptureComponent.h
@@ -1,36 +1,4 @@
-// %BANNER_BEGIN%
-// ---------------------------------------------------------------------
-// %COPYRIGHT_BEGIN%
-//
-// Copyright (c) 2017 Magic Leap, Inc. (COMPANY) All Rights Reserved.
-// Magic Leap, Inc. Confidential and Proprietary
-//
-// NOTICE: All information contained herein is, and remains the property
-// of COMPANY. The intellectual and technical concepts contained herein
-// are proprietary to COMPANY and may be covered by U.S. and Foreign
-// Patents, patents in process, and are protected by trade secret or
-// copyright law. Dissemination of this information or reproduction of
-// this material is strictly forbidden unless prior written permission is
-// obtained from COMPANY. Access to the source code contained herein is
-// hereby forbidden to anyone except current COMPANY employees, managers
-// or contractors who have executed Confidentiality and Non-disclosure
-// agreements explicitly covering such access.
-//
-// The copyright notice above does not evidence any actual or intended
-// publication or disclosure of this source code, which includes
-// information that is confidential and/or proprietary, and is a trade
-// secret, of COMPANY. ANY REPRODUCTION, MODIFICATION, DISTRIBUTION,
-// PUBLIC PERFORMANCE, OR PUBLIC DISPLAY OF OR THROUGH USE OF THIS
-// SOURCE CODE WITHOUT THE EXPRESS WRITTEN CONSENT OF COMPANY IS
-// STRICTLY PROHIBITED, AND IN VIOLATION OF APPLICABLE LAWS AND
-// INTERNATIONAL TREATIES. THE RECEIPT OR POSSESSION OF THIS SOURCE
-// CODE AND/OR RELATED INFORMATION DOES NOT CONVEY OR IMPLY ANY RIGHTS
-// TO REPRODUCE, DISCLOSE OR DISTRIBUTE ITS CONTENTS, OR TO MANUFACTURE,
-// USE, OR SELL ANYTHING THAT IT MAY DESCRIBE, IN WHOLE OR IN PART.
-//
-// %COPYRIGHT_END%
-// --------------------------------------------------------------------
-// %BANNER_END%
+// Copyright 1998-2018 Epic Games, Inc. All Rights Reserved.
 
 #pragma once
 
diff --git a/Engine/Plugins/Lumin/MagicLeap/Source/MagicLeap/Classes/FoundObjectComponent.h b/Engine/Plugins/Lumin/MagicLeap/Source/MagicLeap/Classes/FoundObjectComponent.h
--- a/Engine/Plugins/Lumin/MagicLeap/Source/MagicLeap/Classes/FoundObjectComponent.h
+++ b/Engine/Plugins/Lumin/MagicLeap/Source/MagicLeap/Classes/FoundObjectComponent.h
@@ -1,36 +1,4 @@
-// %BANNER_BEGIN%
-// ---------------------------------------------------------------------
-// %COPYRIGHT_BEGIN%
-//
-// Copyright (c) 2017 Magic Leap, Inc. (COMPANY) All Rights Reserved.
-// Magic Leap, Inc. Confidential and Proprietary
-//
-// NOTICE: All information contained herein is, and remains the property
-// of COMPANY. The intellectual and technical concepts contained herein
-// are proprietary to COMPANY and may be covered by U.S. and Foreign
-// Patents, patents in process, and are protected by trade secret or
-// copyright law. Dissemination of this information or reproduction of
-// this material is strictly forbidden unless prior written permission is
-// obtained from COMPANY. Access to the source code contained herein is
-// hereby forbidden to anyone except current COMPANY employees, managers
-// or contractors who have executed Confidentiality and Non-disclosure
-// agreements explicitly covering such access.
-//
-// The copyright notice above does not evidence any actual or intended
-// publication or disclosure of this source code, which includes
-// information that is confidential and/or proprietary, and is a trade
-// secret, of COMPANY. ANY REPRODUCTION, MODIFICATION, DISTRIBUTION,
-// PUBLIC PERFORMANCE, OR PUBLIC DISPLAY OF OR THROUGH USE OF THIS
-// SOURCE CODE WITHOUT THE EXPRESS WRITTEN CONSENT OF COMPANY IS
-// STRICTLY PROHIBITED, AND IN VIOLATION OF APPLICABLE LAWS AND
-// INTERNATIONAL TREATIES. THE RECEIPT OR POSSESSION OF THIS SOURCE
-// CODE AND/OR RELATED INFORMATION DOES NOT CONVEY OR IMPLY ANY RIGHTS
-// TO REPRODUCE, DISCLOSE OR DISTRIBUTE ITS CONTENTS, OR TO MANUFACTURE,
-// USE, OR SELL ANYTHING THAT IT MAY DESCRIBE, IN WHOLE OR IN PART.
-//
-// %COPYRIGHT_END%
-// --------------------------------------------------------------------
-// %BANNER_END%
+// Copyright 1998-2018 Epic Games, Inc. All Rights Reserved.
 
 #pragma once
 
diff --git a/Engine/Plugins/Lumin/MagicLeap/Source/MagicLeap/Classes/ImageTrackerComponent.h b/Engine/Plugins/Lumin/MagicLeap/Source/MagicLeap/Classes/ImageTrackerComponent.h
--- a/Engine/Plugins/Lumin/MagicLeap/Source/MagicLeap/Classes/ImageTrackerComponent.h
+++ b/Engine/Plugins/Lumin/MagicLeap/Source/MagicLeap/Classes/ImageTrackerComponent.h
@@ -1,36 +1,4 @@
-// %BANNER_BEGIN%
-// ---------------------------------------------------------------------
-// %COPYRIGHT_BEGIN%
-//
-// Copyright (c) 2017 Magic Leap, Inc. (COMPANY) All Rights Reserved.
-// Magic Leap, Inc. Confidential and Proprietary
-//
-// NOTICE: All information contained herein is, and remains the property
-// of COMPANY. The intellectual and technical concepts contained herein
-// are proprietary to COMPANY and may be covered by U.S. and Foreign
-// Patents, patents in process, and are protected by trade secret or
-// copyright law. Dissemination of this information or reproduction of
-// this material is strictly forbidden unless prior written permission is
-// obtained from COMPANY. Access to the source code contained herein is
-// hereby forbidden to anyone except current COMPANY employees, managers
-// or contractors who have executed Confidentiality and Non-disclosure
-// agreements explicitly covering such access.
-//
-// The copyright notice above does not evidence any actual or intended
-// publication or disclosure of this source code, which includes
-// information that is confidential and/or proprietary, and is a trade
-// secret, of COMPANY. ANY REPRODUCTION, MODIFICATION, DISTRIBUTION,
-// PUBLIC PERFORMANCE, OR PUBLIC DISPLAY OF OR THROUGH USE OF THIS
-// SOURCE CODE WITHOUT THE EXPRESS WRITTEN CONSENT OF COMPANY IS
-// STRICTLY PROHIBITED, AND IN VIOLATION OF APPLICABLE LAWS AND
-// INTERNATIONAL TREATIES. THE RECEIPT OR POSSESSION OF THIS SOURCE
-// CODE AND/OR RELATED INFORMATION DOES NOT CONVEY OR IMPLY ANY RIGHTS
-// TO REPRODUCE, DISCLOSE OR DISTRIBUTE ITS CONTENTS, OR TO MANUFACTURE,
-// USE, OR SELL ANYTHING THAT IT MAY DESCRIBE, IN WHOLE OR IN PART.
-//
-// %COPYRIGHT_END%
-// --------------------------------------------------------------------
-// %BANNER_END%
+// Copyright 1998-2018 Epic Games, Inc. All Rights Reserved.
 
 #pragma once
 
diff --git a/Engine/Plugins/Lumin/MagicLeap/Source/MagicLeap/Classes/MagicLeapHMDFunctionLibrary.h b/Engine/Plugins/Lumin/MagicLeap/Source/MagicLeap/Classes/MagicLeapHMDFunctionLibrary.h
--- a/Engine/Plugins/Lumin/MagicLeap/Source/MagicLeap/Classes/MagicLeapHMDFunctionLibrary.h
+++ b/Engine/Plugins/Lumin/MagicLeap/Source/MagicLeap/Classes/MagicLeapHMDFunctionLibrary.h
@@ -1,36 +1,4 @@
-// %BANNER_BEGIN%
-// ---------------------------------------------------------------------
-// %COPYRIGHT_BEGIN%
-//
-// Copyright (c) 2017 Magic Leap, Inc. (COMPANY) All Rights Reserved.
-// Magic Leap, Inc. Confidential and Proprietary
-//
-// NOTICE: All information contained herein is, and remains the property
-// of COMPANY. The intellectual and technical concepts contained herein
-// are proprietary to COMPANY and may be covered by U.S. and Foreign
-// Patents, patents in process, and are protected by trade secret or
-// copyright law. Dissemination of this information or reproduction of
-// this material is strictly forbidden unless prior written permission is
-// obtained from COMPANY. Access to the source code contained herein is
-// hereby forbidden to anyone except current COMPANY employees, managers
-// or contractors who have executed Confidentiality and Non-disclosure
-// agreements explicitly covering such access.
-//
-// The copyright notice above does not evidence any actual or intended
-// publication or disclosure of this source code, which includes
-// information that is confidential and/or proprietary, and is a trade
-// secret, of COMPANY. ANY REPRODUCTION, MODIFICATION, DISTRIBUTION,
-// PUBLIC PERFORMANCE, OR PUBLIC DISPLAY OF OR THROUGH USE OF THIS
-// SOURCE CODE WITHOUT THE EXPRESS WRITTEN CONSENT OF COMPANY IS
-// STRICTLY PROHIBITED, AND IN VIOLATION OF APPLICABLE LAWS AND
-// INTERNATIONAL TREATIES. THE RECEIPT OR POSSESSION OF THIS SOURCE
-// CODE AND/OR RELATED INFORMATION DOES NOT CONVEY OR IMPLY ANY RIGHTS
-// TO REPRODUCE, DISCLOSE OR DISTRIBUTE ITS CONTENTS, OR TO MANUFACTURE,
-// USE, OR SELL ANYTHING THAT IT MAY DESCRIBE, IN WHOLE OR IN PART.
-//
-// %COPYRIGHT_END%
-// --------------------------------------------------------------------
-// %BANNER_END%
+// Copyright 1998-2018 Epic Games, Inc. All Rights Reserved.
 
 #pragma once
 
diff --git a/Engine/Plugins/Lumin/MagicLeap/Source/MagicLeap/Classes/MagicLeapSettings.h b/Engine/Plugins/Lumin/MagicLeap/Source/MagicLeap/Classes/MagicLeapSettings.h
--- a/Engine/Plugins/Lumin/MagicLeap/Source/MagicLeap/Classes/MagicLeapSettings.h
+++ b/Engine/Plugins/Lumin/MagicLeap/Source/MagicLeap/Classes/MagicLeapSettings.h
@@ -1,3 +1,5 @@
+// Copyright 1998-2018 Epic Games, Inc. All Rights Reserved.
+
 #pragma once
 
 #include "CoreMinimal.h"
diff --git a/Engine/Plugins/Lumin/MagicLeap/Source/MagicLeap/Classes/MeshTrackerComponent.h b/Engine/Plugins/Lumin/MagicLeap/Source/MagicLeap/Classes/MeshTrackerComponent.h
--- a/Engine/Plugins/Lumin/MagicLeap/Source/MagicLeap/Classes/MeshTrackerComponent.h
+++ b/Engine/Plugins/Lumin/MagicLeap/Source/MagicLeap/Classes/MeshTrackerComponent.h
@@ -1,36 +1,4 @@
-// %BANNER_BEGIN%
-// ---------------------------------------------------------------------
-// %COPYRIGHT_BEGIN%
-//
-// Copyright (c) 2017 Magic Leap, Inc. (COMPANY) All Rights Reserved.
-// Magic Leap, Inc. Confidential and Proprietary
-//
-// NOTICE: All information contained herein is, and remains the property
-// of COMPANY. The intellectual and technical concepts contained herein
-// are proprietary to COMPANY and may be covered by U.S. and Foreign
-// Patents, patents in process, and are protected by trade secret or
-// copyright law. Dissemination of this information or reproduction of
-// this material is strictly forbidden unless prior written permission is
-// obtained from COMPANY. Access to the source code contained herein is
-// hereby forbidden to anyone except current COMPANY employees, managers
-// or contractors who have executed Confidentiality and Non-disclosure
-// agreements explicitly covering such access.
-//
-// The copyright notice above does not evidence any actual or intended
-// publication or disclosure of this source code, which includes
-// information that is confidential and/or proprietary, and is a trade
-// secret, of COMPANY. ANY REPRODUCTION, MODIFICATION, DISTRIBUTION,
-// PUBLIC PERFORMANCE, OR PUBLIC DISPLAY OF OR THROUGH USE OF THIS
-// SOURCE CODE WITHOUT THE EXPRESS WRITTEN CONSENT OF COMPANY IS
-// STRICTLY PROHIBITED, AND IN VIOLATION OF APPLICABLE LAWS AND
-// INTERNATIONAL TREATIES. THE RECEIPT OR POSSESSION OF THIS SOURCE
-// CODE AND/OR RELATED INFORMATION DOES NOT CONVEY OR IMPLY ANY RIGHTS
-// TO REPRODUCE, DISCLOSE OR DISTRIBUTE ITS CONTENTS, OR TO MANUFACTURE,
-// USE, OR SELL ANYTHING THAT IT MAY DESCRIBE, IN WHOLE OR IN PART.
-//
-// %COPYRIGHT_END%
-// --------------------------------------------------------------------
-// %BANNER_END%
+// Copyright 1998-2018 Epic Games, Inc. All Rights Reserved.
 
 #pragma once
 
diff --git a/Engine/Plugins/Lumin/MagicLeap/Source/MagicLeap/Classes/PlanesComponent.h b/Engine/Plugins/Lumin/MagicLeap/Source/MagicLeap/Classes/PlanesComponent.h
--- a/Engine/Plugins/Lumin/MagicLeap/Source/MagicLeap/Classes/PlanesComponent.h
+++ b/Engine/Plugins/Lumin/MagicLeap/Source/MagicLeap/Classes/PlanesComponent.h
@@ -1,36 +1,4 @@
-// %BANNER_BEGIN%
-// ---------------------------------------------------------------------
-// %COPYRIGHT_BEGIN%
-//
-// Copyright (c) 2017 Magic Leap, Inc. (COMPANY) All Rights Reserved.
-// Magic Leap, Inc. Confidential and Proprietary
-//
-// NOTICE: All information contained herein is, and remains the property
-// of COMPANY. The intellectual and technical concepts contained herein
-// are proprietary to COMPANY and may be covered by U.S. and Foreign
-// Patents, patents in process, and are protected by trade secret or
-// copyright law. Dissemination of this information or reproduction of
-// this material is strictly forbidden unless prior written permission is
-// obtained from COMPANY. Access to the source code contained herein is
-// hereby forbidden to anyone except current COMPANY employees, managers
-// or contractors who have executed Confidentiality and Non-disclosure
-// agreements explicitly covering such access.
-//
-// The copyright notice above does not evidence any actual or intended
-// publication or disclosure of this source code, which includes
-// information that is confidential and/or proprietary, and is a trade
-// secret, of COMPANY. ANY REPRODUCTION, MODIFICATION, DISTRIBUTION,
-// PUBLIC PERFORMANCE, OR PUBLIC DISPLAY OF OR THROUGH USE OF THIS
-// SOURCE CODE WITHOUT THE EXPRESS WRITTEN CONSENT OF COMPANY IS
-// STRICTLY PROHIBITED, AND IN VIOLATION OF APPLICABLE LAWS AND
-// INTERNATIONAL TREATIES. THE RECEIPT OR POSSESSION OF THIS SOURCE
-// CODE AND/OR RELATED INFORMATION DOES NOT CONVEY OR IMPLY ANY RIGHTS
-// TO REPRODUCE, DISCLOSE OR DISTRIBUTE ITS CONTENTS, OR TO MANUFACTURE,
-// USE, OR SELL ANYTHING THAT IT MAY DESCRIBE, IN WHOLE OR IN PART.
-//
-// %COPYRIGHT_END%
-// --------------------------------------------------------------------
-// %BANNER_END%
+// Copyright 1998-2018 Epic Games, Inc. All Rights Reserved.
 
 #pragma once
 
diff --git a/Engine/Plugins/Lumin/MagicLeap/Source/MagicLeap/Classes/RaycastComponent.h b/Engine/Plugins/Lumin/MagicLeap/Source/MagicLeap/Classes/RaycastComponent.h
--- a/Engine/Plugins/Lumin/MagicLeap/Source/MagicLeap/Classes/RaycastComponent.h
+++ b/Engine/Plugins/Lumin/MagicLeap/Source/MagicLeap/Classes/RaycastComponent.h
@@ -1,36 +1,4 @@
-// %BANNER_BEGIN%
-// ---------------------------------------------------------------------
-// %COPYRIGHT_BEGIN%
-//
-// Copyright (c) 2017 Magic Leap, Inc. (COMPANY) All Rights Reserved.
-// Magic Leap, Inc. Confidential and Proprietary
-//
-// NOTICE: All information contained herein is, and remains the property
-// of COMPANY. The intellectual and technical concepts contained herein
-// are proprietary to COMPANY and may be covered by U.S. and Foreign
-// Patents, patents in process, and are protected by trade secret or
-// copyright law. Dissemination of this information or reproduction of
-// this material is strictly forbidden unless prior written permission is
-// obtained from COMPANY. Access to the source code contained herein is
-// hereby forbidden to anyone except current COMPANY employees, managers
-// or contractors who have executed Confidentiality and Non-disclosure
-// agreements explicitly covering such access.
-//
-// The copyright notice above does not evidence any actual or intended
-// publication or disclosure of this source code, which includes
-// information that is confidential and/or proprietary, and is a trade
-// secret, of COMPANY. ANY REPRODUCTION, MODIFICATION, DISTRIBUTION,
-// PUBLIC PERFORMANCE, OR PUBLIC DISPLAY OF OR THROUGH USE OF THIS
-// SOURCE CODE WITHOUT THE EXPRESS WRITTEN CONSENT OF COMPANY IS
-// STRICTLY PROHIBITED, AND IN VIOLATION OF APPLICABLE LAWS AND
-// INTERNATIONAL TREATIES. THE RECEIPT OR POSSESSION OF THIS SOURCE
-// CODE AND/OR RELATED INFORMATION DOES NOT CONVEY OR IMPLY ANY RIGHTS
-// TO REPRODUCE, DISCLOSE OR DISTRIBUTE ITS CONTENTS, OR TO MANUFACTURE,
-// USE, OR SELL ANYTHING THAT IT MAY DESCRIBE, IN WHOLE OR IN PART.
-//
-// %COPYRIGHT_END%
-// --------------------------------------------------------------------
-// %BANNER_END%
+// Copyright 1998-2018 Epic Games, Inc. All Rights Reserved.
 
 #pragma once
 
diff --git a/Engine/Plugins/Lumin/MagicLeap/Source/MagicLeap/MagicLeap.Build.cs b/Engine/Plugins/Lumin/MagicLeap/Source/MagicLeap/MagicLeap.Build.cs
--- a/Engine/Plugins/Lumin/MagicLeap/Source/MagicLeap/MagicLeap.Build.cs
+++ b/Engine/Plugins/Lumin/MagicLeap/Source/MagicLeap/MagicLeap.Build.cs
@@ -1,36 +1,4 @@
-// %BANNER_BEGIN%
-// ---------------------------------------------------------------------
-// %COPYRIGHT_BEGIN%
-//
-// Copyright (c) 2017 Magic Leap, Inc. (COMPANY) All Rights Reserved.
-// Magic Leap, Inc. Confidential and Proprietary
-//
-// NOTICE: All information contained herein is, and remains the property
-// of COMPANY. The intellectual and technical concepts contained herein
-// are proprietary to COMPANY and may be covered by U.S. and Foreign
-// Patents, patents in process, and are protected by trade secret or
-// copyright law. Dissemination of this information or reproduction of
-// this material is strictly forbidden unless prior written permission is
-// obtained from COMPANY. Access to the source code contained herein is
-// hereby forbidden to anyone except current COMPANY employees, managers
-// or contractors who have executed Confidentiality and Non-disclosure
-// agreements explicitly covering such access.
-//
-// The copyright notice above does not evidence any actual or intended
-// publication or disclosure of this source code, which includes
-// information that is confidential and/or proprietary, and is a trade
-// secret, of COMPANY. ANY REPRODUCTION, MODIFICATION, DISTRIBUTION,
-// PUBLIC PERFORMANCE, OR PUBLIC DISPLAY OF OR THROUGH USE OF THIS
-// SOURCE CODE WITHOUT THE EXPRESS WRITTEN CONSENT OF COMPANY IS
-// STRICTLY PROHIBITED, AND IN VIOLATION OF APPLICABLE LAWS AND
-// INTERNATIONAL TREATIES. THE RECEIPT OR POSSESSION OF THIS SOURCE
-// CODE AND/OR RELATED INFORMATION DOES NOT CONVEY OR IMPLY ANY RIGHTS
-// TO REPRODUCE, DISCLOSE OR DISTRIBUTE ITS CONTENTS, OR TO MANUFACTURE,
-// USE, OR SELL ANYTHING THAT IT MAY DESCRIBE, IN WHOLE OR IN PART.
-//
-// %COPYRIGHT_END%
-// --------------------------------------------------------------------
-// %BANNER_END%
+// Copyright 1998-2018 Epic Games, Inc. All Rights Reserved.
 
 using System.IO;
 
diff --git a/Engine/Plugins/Lumin/MagicLeap/Source/MagicLeap/Private/AppEventHandler.cpp b/Engine/Plugins/Lumin/MagicLeap/Source/MagicLeap/Private/AppEventHandler.cpp
--- a/Engine/Plugins/Lumin/MagicLeap/Source/MagicLeap/Private/AppEventHandler.cpp
+++ b/Engine/Plugins/Lumin/MagicLeap/Source/MagicLeap/Private/AppEventHandler.cpp
@@ -1,36 +1,4 @@
-// %BANNER_BEGIN%
-// ---------------------------------------------------------------------
-// %COPYRIGHT_BEGIN%
-//
-// Copyright (c) 201x Magic Leap, Inc. (COMPANY) All Rights Reserved.
-// Magic Leap, Inc. Confidential and Proprietary
-//
-// NOTICE: All information contained herein is, and remains the property
-// of COMPANY. The intellectual and technical concepts contained herein
-// are proprietary to COMPANY and may be covered by U.S. and Foreign
-// Patents, patents in process, and are protected by trade secret or
-// copyright law. Dissemination of this information or reproduction of
-// this material is strictly forbidden unless prior written permission is
-// obtained from COMPANY. Access to the source code contained herein is
-// hereby forbidden to anyone except current COMPANY employees, managers
-// or contractors who have executed Confidentiality and Non-disclosure
-// agreements explicitly covering such access.
-//
-// The copyright notice above does not evidence any actual or intended
-// publication or disclosure of this source code, which includes
-// information that is confidential and/or proprietary, and is a trade
-// secret, of COMPANY. ANY REPRODUCTION, MODIFICATION, DISTRIBUTION,
-// PUBLIC PERFORMANCE, OR PUBLIC DISPLAY OF OR THROUGH USE OF THIS
-// SOURCE CODE WITHOUT THE EXPRESS WRITTEN CONSENT OF COMPANY IS
-// STRICTLY PROHIBITED, AND IN VIOLATION OF APPLICABLE LAWS AND
-// INTERNATIONAL TREATIES. THE RECEIPT OR POSSESSION OF THIS SOURCE
-// CODE AND/OR RELATED INFORMATION DOES NOT CONVEY OR IMPLY ANY RIGHTS
-// TO REPRODUCE, DISCLOSE OR DISTRIBUTE ITS CONTENTS, OR TO MANUFACTURE,
-// USE, OR SELL ANYTHING THAT IT MAY DESCRIBE, IN WHOLE OR IN PART.
-//
-// %COPYRIGHT_END%
-// --------------------------------------------------------------------
-// %BANNER_END%
+// Copyright 1998-2018 Epic Games, Inc. All Rights Reserved.
 
 #include "AppEventHandler.h"
 #include "AppFramework.h"
diff --git a/Engine/Plugins/Lumin/MagicLeap/Source/MagicLeap/Private/AppEventHandler.h b/Engine/Plugins/Lumin/MagicLeap/Source/MagicLeap/Private/AppEventHandler.h
--- a/Engine/Plugins/Lumin/MagicLeap/Source/MagicLeap/Private/AppEventHandler.h
+++ b/Engine/Plugins/Lumin/MagicLeap/Source/MagicLeap/Private/AppEventHandler.h
@@ -1,36 +1,4 @@
-// %BANNER_BEGIN%
-// ---------------------------------------------------------------------
-// %COPYRIGHT_BEGIN%
-//
-// Copyright (c) 2017 Magic Leap, Inc. (COMPANY) All Rights Reserved.
-// Magic Leap, Inc. Confidential and Proprietary
-//
-// NOTICE: All information contained herein is, and remains the property
-// of COMPANY. The intellectual and technical concepts contained herein
-// are proprietary to COMPANY and may be covered by U.S. and Foreign
-// Patents, patents in process, and are protected by trade secret or
-// copyright law. Dissemination of this information or reproduction of
-// this material is strictly forbidden unless prior written permission is
-// obtained from COMPANY. Access to the source code contained herein is
-// hereby forbidden to anyone except current COMPANY employees, managers
-// or contractors who have executed Confidentiality and Non-disclosure
-// agreements explicitly covering such access.
-//
-// The copyright notice above does not evidence any actual or intended
-// publication or disclosure of this source code, which includes
-// information that is confidential and/or proprietary, and is a trade
-// secret, of COMPANY. ANY REPRODUCTION, MODIFICATION, DISTRIBUTION,
-// PUBLIC PERFORMANCE, OR PUBLIC DISPLAY OF OR THROUGH USE OF THIS
-// SOURCE CODE WITHOUT THE EXPRESS WRITTEN CONSENT OF COMPANY IS
-// STRICTLY PROHIBITED, AND IN VIOLATION OF APPLICABLE LAWS AND
-// INTERNATIONAL TREATIES. THE RECEIPT OR POSSESSION OF THIS SOURCE
-// CODE AND/OR RELATED INFORMATION DOES NOT CONVEY OR IMPLY ANY RIGHTS
-// TO REPRODUCE, DISCLOSE OR DISTRIBUTE ITS CONTENTS, OR TO MANUFACTURE,
-// USE, OR SELL ANYTHING THAT IT MAY DESCRIBE, IN WHOLE OR IN PART.
-//
-// %COPYRIGHT_END%
-// --------------------------------------------------------------------
-// %BANNER_END%
+// Copyright 1998-2018 Epic Games, Inc. All Rights Reserved.
 
 #pragma once
 
diff --git a/Engine/Plugins/Lumin/MagicLeap/Source/MagicLeap/Private/AppFramework.cpp b/Engine/Plugins/Lumin/MagicLeap/Source/MagicLeap/Private/AppFramework.cpp
--- a/Engine/Plugins/Lumin/MagicLeap/Source/MagicLeap/Private/AppFramework.cpp
+++ b/Engine/Plugins/Lumin/MagicLeap/Source/MagicLeap/Private/AppFramework.cpp
@@ -1,36 +1,4 @@
-// %BANNER_BEGIN%
-// ---------------------------------------------------------------------
-// %COPYRIGHT_BEGIN%
-//
-// Copyright (c) 2017 Magic Leap, Inc. (COMPANY) All Rights Reserved.
-// Magic Leap, Inc. Confidential and Proprietary
-//
-// NOTICE: All information contained herein is, and remains the property
-// of COMPANY. The intellectual and technical concepts contained herein
-// are proprietary to COMPANY and may be covered by U.S. and Foreign
-// Patents, patents in process, and are protected by trade secret or
-// copyright law. Dissemination of this information or reproduction of
-// this material is strictly forbidden unless prior written permission is
-// obtained from COMPANY. Access to the source code contained herein is
-// hereby forbidden to anyone except current COMPANY employees, managers
-// or contractors who have executed Confidentiality and Non-disclosure
-// agreements explicitly covering such access.
-//
-// The copyright notice above does not evidence any actual or intended
-// publication or disclosure of this source code, which includes
-// information that is confidential and/or proprietary, and is a trade
-// secret, of COMPANY. ANY REPRODUCTION, MODIFICATION, DISTRIBUTION,
-// PUBLIC PERFORMANCE, OR PUBLIC DISPLAY OF OR THROUGH USE OF THIS
-// SOURCE CODE WITHOUT THE EXPRESS WRITTEN CONSENT OF COMPANY IS
-// STRICTLY PROHIBITED, AND IN VIOLATION OF APPLICABLE LAWS AND
-// INTERNATIONAL TREATIES. THE RECEIPT OR POSSESSION OF THIS SOURCE
-// CODE AND/OR RELATED INFORMATION DOES NOT CONVEY OR IMPLY ANY RIGHTS
-// TO REPRODUCE, DISCLOSE OR DISTRIBUTE ITS CONTENTS, OR TO MANUFACTURE,
-// USE, OR SELL ANYTHING THAT IT MAY DESCRIBE, IN WHOLE OR IN PART.
-//
-// %COPYRIGHT_END%
-// --------------------------------------------------------------------
-// %BANNER_END%
+// Copyright 1998-2018 Epic Games, Inc. All Rights Reserved.
 
 #include "AppFramework.h"
 #include "MagicLeapHMD.h"
diff --git a/Engine/Plugins/Lumin/MagicLeap/Source/MagicLeap/Private/AppFramework.h b/Engine/Plugins/Lumin/MagicLeap/Source/MagicLeap/Private/AppFramework.h
--- a/Engine/Plugins/Lumin/MagicLeap/Source/MagicLeap/Private/AppFramework.h
+++ b/Engine/Plugins/Lumin/MagicLeap/Source/MagicLeap/Private/AppFramework.h
@@ -1,36 +1,4 @@
-// %BANNER_BEGIN%
-// ---------------------------------------------------------------------
-// %COPYRIGHT_BEGIN%
-//
-// Copyright (c) 2017 Magic Leap, Inc. (COMPANY) All Rights Reserved.
-// Magic Leap, Inc. Confidential and Proprietary
-//
-// NOTICE: All information contained herein is, and remains the property
-// of COMPANY. The intellectual and technical concepts contained herein
-// are proprietary to COMPANY and may be covered by U.S. and Foreign
-// Patents, patents in process, and are protected by trade secret or
-// copyright law. Dissemination of this information or reproduction of
-// this material is strictly forbidden unless prior written permission is
-// obtained from COMPANY. Access to the source code contained herein is
-// hereby forbidden to anyone except current COMPANY employees, managers
-// or contractors who have executed Confidentiality and Non-disclosure
-// agreements explicitly covering such access.
-//
-// The copyright notice above does not evidence any actual or intended
-// publication or disclosure of this source code, which includes
-// information that is confidential and/or proprietary, and is a trade
-// secret, of COMPANY. ANY REPRODUCTION, MODIFICATION, DISTRIBUTION,
-// PUBLIC PERFORMANCE, OR PUBLIC DISPLAY OF OR THROUGH USE OF THIS
-// SOURCE CODE WITHOUT THE EXPRESS WRITTEN CONSENT OF COMPANY IS
-// STRICTLY PROHIBITED, AND IN VIOLATION OF APPLICABLE LAWS AND
-// INTERNATIONAL TREATIES. THE RECEIPT OR POSSESSION OF THIS SOURCE
-// CODE AND/OR RELATED INFORMATION DOES NOT CONVEY OR IMPLY ANY RIGHTS
-// TO REPRODUCE, DISCLOSE OR DISTRIBUTE ITS CONTENTS, OR TO MANUFACTURE,
-// USE, OR SELL ANYTHING THAT IT MAY DESCRIBE, IN WHOLE OR IN PART.
-//
-// %COPYRIGHT_END%
-// --------------------------------------------------------------------
-// %BANNER_END%
+// Copyright 1998-2018 Epic Games, Inc. All Rights Reserved.
 
 #pragma once
 
diff --git a/Engine/Plugins/Lumin/MagicLeap/Source/MagicLeap/Private/AsyncDestroyer.cpp b/Engine/Plugins/Lumin/MagicLeap/Source/MagicLeap/Private/AsyncDestroyer.cpp
--- a/Engine/Plugins/Lumin/MagicLeap/Source/MagicLeap/Private/AsyncDestroyer.cpp
+++ b/Engine/Plugins/Lumin/MagicLeap/Source/MagicLeap/Private/AsyncDestroyer.cpp
@@ -1,36 +1,4 @@
-// %BANNER_BEGIN%
-// ---------------------------------------------------------------------
-// %COPYRIGHT_BEGIN%
-//
-// Copyright (c) 201x Magic Leap, Inc. (COMPANY) All Rights Reserved.
-// Magic Leap, Inc. Confidential and Proprietary
-//
-// NOTICE: All information contained herein is, and remains the property
-// of COMPANY. The intellectual and technical concepts contained herein
-// are proprietary to COMPANY and may be covered by U.S. and Foreign
-// Patents, patents in process, and are protected by trade secret or
-// copyright law. Dissemination of this information or reproduction of
-// this material is strictly forbidden unless prior written permission is
-// obtained from COMPANY. Access to the source code contained herein is
-// hereby forbidden to anyone except current COMPANY employees, managers
-// or contractors who have executed Confidentiality and Non-disclosure
-// agreements explicitly covering such access.
-//
-// The copyright notice above does not evidence any actual or intended
-// publication or disclosure of this source code, which includes
-// information that is confidential and/or proprietary, and is a trade
-// secret, of COMPANY. ANY REPRODUCTION, MODIFICATION, DISTRIBUTION,
-// PUBLIC PERFORMANCE, OR PUBLIC DISPLAY OF OR THROUGH USE OF THIS
-// SOURCE CODE WITHOUT THE EXPRESS WRITTEN CONSENT OF COMPANY IS
-// STRICTLY PROHIBITED, AND IN VIOLATION OF APPLICABLE LAWS AND
-// INTERNATIONAL TREATIES. THE RECEIPT OR POSSESSION OF THIS SOURCE
-// CODE AND/OR RELATED INFORMATION DOES NOT CONVEY OR IMPLY ANY RIGHTS
-// TO REPRODUCE, DISCLOSE OR DISTRIBUTE ITS CONTENTS, OR TO MANUFACTURE,
-// USE, OR SELL ANYTHING THAT IT MAY DESCRIBE, IN WHOLE OR IN PART.
-//
-// %COPYRIGHT_END%
-// --------------------------------------------------------------------
-// %BANNER_END%
+// Copyright 1998-2018 Epic Games, Inc. All Rights Reserved.
 
 #include "AsyncDestroyer.h"
 #include "Engine/Engine.h"
diff --git a/Engine/Plugins/Lumin/MagicLeap/Source/MagicLeap/Private/AsyncDestroyer.h b/Engine/Plugins/Lumin/MagicLeap/Source/MagicLeap/Private/AsyncDestroyer.h
--- a/Engine/Plugins/Lumin/MagicLeap/Source/MagicLeap/Private/AsyncDestroyer.h
+++ b/Engine/Plugins/Lumin/MagicLeap/Source/MagicLeap/Private/AsyncDestroyer.h
@@ -1,36 +1,4 @@
-// %BANNER_BEGIN%
-// ---------------------------------------------------------------------
-// %COPYRIGHT_BEGIN%
-//
-// Copyright (c) 201x Magic Leap, Inc. (COMPANY) All Rights Reserved.
-// Magic Leap, Inc. Confidential and Proprietary
-//
-// NOTICE: All information contained herein is, and remains the property
-// of COMPANY. The intellectual and technical concepts contained herein
-// are proprietary to COMPANY and may be covered by U.S. and Foreign
-// Patents, patents in process, and are protected by trade secret or
-// copyright law. Dissemination of this information or reproduction of
-// this material is strictly forbidden unless prior written permission is
-// obtained from COMPANY. Access to the source code contained herein is
-// hereby forbidden to anyone except current COMPANY employees, managers
-// or contractors who have executed Confidentiality and Non-disclosure
-// agreements explicitly covering such access.
-//
-// The copyright notice above does not evidence any actual or intended
-// publication or disclosure of this source code, which includes
-// information that is confidential and/or proprietary, and is a trade
-// secret, of COMPANY. ANY REPRODUCTION, MODIFICATION, DISTRIBUTION,
-// PUBLIC PERFORMANCE, OR PUBLIC DISPLAY OF OR THROUGH USE OF THIS
-// SOURCE CODE WITHOUT THE EXPRESS WRITTEN CONSENT OF COMPANY IS
-// STRICTLY PROHIBITED, AND IN VIOLATION OF APPLICABLE LAWS AND
-// INTERNATIONAL TREATIES. THE RECEIPT OR POSSESSION OF THIS SOURCE
-// CODE AND/OR RELATED INFORMATION DOES NOT CONVEY OR IMPLY ANY RIGHTS
-// TO REPRODUCE, DISCLOSE OR DISTRIBUTE ITS CONTENTS, OR TO MANUFACTURE,
-// USE, OR SELL ANYTHING THAT IT MAY DESCRIBE, IN WHOLE OR IN PART.
-//
-// %COPYRIGHT_END%
-// --------------------------------------------------------------------
-// %BANNER_END%
+// Copyright 1998-2018 Epic Games, Inc. All Rights Reserved.
 
 #pragma once
 
diff --git a/Engine/Plugins/Lumin/MagicLeap/Source/MagicLeap/Private/CameraCaptureComponent.cpp b/Engine/Plugins/Lumin/MagicLeap/Source/MagicLeap/Private/CameraCaptureComponent.cpp
--- a/Engine/Plugins/Lumin/MagicLeap/Source/MagicLeap/Private/CameraCaptureComponent.cpp
+++ b/Engine/Plugins/Lumin/MagicLeap/Source/MagicLeap/Private/CameraCaptureComponent.cpp
@@ -1,4 +1,4 @@
-// Fill out your copyright notice in the Description page of Project Settings.
+// Copyright 1998-2018 Epic Games, Inc. All Rights Reserved.
 
 #include "CameraCaptureComponent.h"
 #include "CoreMinimal.h"
diff --git a/Engine/Plugins/Lumin/MagicLeap/Source/MagicLeap/Private/FoundObjectComponent.cpp b/Engine/Plugins/Lumin/MagicLeap/Source/MagicLeap/Private/FoundObjectComponent.cpp
--- a/Engine/Plugins/Lumin/MagicLeap/Source/MagicLeap/Private/FoundObjectComponent.cpp
+++ b/Engine/Plugins/Lumin/MagicLeap/Source/MagicLeap/Private/FoundObjectComponent.cpp
@@ -1,36 +1,4 @@
-// %BANNER_BEGIN%
-// ---------------------------------------------------------------------
-// %COPYRIGHT_BEGIN%
-//
-// Copyright (c) 2017 Magic Leap, Inc. (COMPANY) All Rights Reserved.
-// Magic Leap, Inc. Confidential and Proprietary
-//
-// NOTICE: All information contained herein is, and remains the property
-// of COMPANY. The intellectual and technical concepts contained herein
-// are proprietary to COMPANY and may be covered by U.S. and Foreign
-// Patents, patents in process, and are protected by trade secret or
-// copyright law. Dissemination of this information or reproduction of
-// this material is strictly forbidden unless prior written permission is
-// obtained from COMPANY. Access to the source code contained herein is
-// hereby forbidden to anyone except current COMPANY employees, managers
-// or contractors who have executed Confidentiality and Non-disclosure
-// agreements explicitly covering such access.
-//
-// The copyright notice above does not evidence any actual or intended
-// publication or disclosure of this source code, which includes
-// information that is confidential and/or proprietary, and is a trade
-// secret, of COMPANY. ANY REPRODUCTION, MODIFICATION, DISTRIBUTION,
-// PUBLIC PERFORMANCE, OR PUBLIC DISPLAY OF OR THROUGH USE OF THIS
-// SOURCE CODE WITHOUT THE EXPRESS WRITTEN CONSENT OF COMPANY IS
-// STRICTLY PROHIBITED, AND IN VIOLATION OF APPLICABLE LAWS AND
-// INTERNATIONAL TREATIES. THE RECEIPT OR POSSESSION OF THIS SOURCE
-// CODE AND/OR RELATED INFORMATION DOES NOT CONVEY OR IMPLY ANY RIGHTS
-// TO REPRODUCE, DISCLOSE OR DISTRIBUTE ITS CONTENTS, OR TO MANUFACTURE,
-// USE, OR SELL ANYTHING THAT IT MAY DESCRIBE, IN WHOLE OR IN PART.
-//
-// %COPYRIGHT_END%
-// --------------------------------------------------------------------
-// %BANNER_END%
+// Copyright 1998-2018 Epic Games, Inc. All Rights Reserved.
 
 #include "FoundObjectComponent.h"
 #include "Components/BoxComponent.h"
diff --git a/Engine/Plugins/Lumin/MagicLeap/Source/MagicLeap/Private/ImageTrackerComponent.cpp b/Engine/Plugins/Lumin/MagicLeap/Source/MagicLeap/Private/ImageTrackerComponent.cpp
--- a/Engine/Plugins/Lumin/MagicLeap/Source/MagicLeap/Private/ImageTrackerComponent.cpp
+++ b/Engine/Plugins/Lumin/MagicLeap/Source/MagicLeap/Private/ImageTrackerComponent.cpp
@@ -1,36 +1,4 @@
-// %BANNER_BEGIN%
-// ---------------------------------------------------------------------
-// %COPYRIGHT_BEGIN%
-//
-// Copyright (c) 2017 Magic Leap, Inc. (COMPANY) All Rights Reserved.
-// Magic Leap, Inc. Confidential and Proprietary
-//
-// NOTICE: All information contained herein is, and remains the property
-// of COMPANY. The intellectual and technical concepts contained herein
-// are proprietary to COMPANY and may be covered by U.S. and Foreign
-// Patents, patents in process, and are protected by trade secret or
-// copyright law. Dissemination of this information or reproduction of
-// this material is strictly forbidden unless prior written permission is
-// obtained from COMPANY. Access to the source code contained herein is
-// hereby forbidden to anyone except current COMPANY employees, managers
-// or contractors who have executed Confidentiality and Non-disclosure
-// agreements explicitly covering such access.
-//
-// The copyright notice above does not evidence any actual or intended
-// publication or disclosure of this source code, which includes
-// information that is confidential and/or proprietary, and is a trade
-// secret, of COMPANY. ANY REPRODUCTION, MODIFICATION, DISTRIBUTION,
-// PUBLIC PERFORMANCE, OR PUBLIC DISPLAY OF OR THROUGH USE OF THIS
-// SOURCE CODE WITHOUT THE EXPRESS WRITTEN CONSENT OF COMPANY IS
-// STRICTLY PROHIBITED, AND IN VIOLATION OF APPLICABLE LAWS AND
-// INTERNATIONAL TREATIES. THE RECEIPT OR POSSESSION OF THIS SOURCE
-// CODE AND/OR RELATED INFORMATION DOES NOT CONVEY OR IMPLY ANY RIGHTS
-// TO REPRODUCE, DISCLOSE OR DISTRIBUTE ITS CONTENTS, OR TO MANUFACTURE,
-// USE, OR SELL ANYTHING THAT IT MAY DESCRIBE, IN WHOLE OR IN PART.
-//
-// %COPYRIGHT_END%
-// --------------------------------------------------------------------
-// %BANNER_END%
+// Copyright 1998-2018 Epic Games, Inc. All Rights Reserved.
 
 #include "ImageTrackerComponent.h"
 #include "MagicLeapHMD.h"
diff --git a/Engine/Plugins/Lumin/MagicLeap/Source/MagicLeap/Private/MagicLeapCustomPresent.cpp b/Engine/Plugins/Lumin/MagicLeap/Source/MagicLeap/Private/MagicLeapCustomPresent.cpp
--- a/Engine/Plugins/Lumin/MagicLeap/Source/MagicLeap/Private/MagicLeapCustomPresent.cpp
+++ b/Engine/Plugins/Lumin/MagicLeap/Source/MagicLeap/Private/MagicLeapCustomPresent.cpp
@@ -1,36 +1,4 @@
-// %BANNER_BEGIN%
-// ---------------------------------------------------------------------
-// %COPYRIGHT_BEGIN%
-//
-// Copyright (c) 2017 Magic Leap, Inc. (COMPANY) All Rights Reserved.
-// Magic Leap, Inc. Confidential and Proprietary
-//
-// NOTICE: All information contained herein is, and remains the property
-// of COMPANY. The intellectual and technical concepts contained herein
-// are proprietary to COMPANY and may be covered by U.S. and Foreign
-// Patents, patents in process, and are protected by trade secret or
-// copyright law. Dissemination of this information or reproduction of
-// this material is strictly forbidden unless prior written permission is
-// obtained from COMPANY. Access to the source code contained herein is
-// hereby forbidden to anyone except current COMPANY employees, managers
-// or contractors who have executed Confidentiality and Non-disclosure
-// agreements explicitly covering such access.
-//
-// The copyright notice above does not evidence any actual or intended
-// publication or disclosure of this source code, which includes
-// information that is confidential and/or proprietary, and is a trade
-// secret, of COMPANY. ANY REPRODUCTION, MODIFICATION, DISTRIBUTION,
-// PUBLIC PERFORMANCE, OR PUBLIC DISPLAY OF OR THROUGH USE OF THIS
-// SOURCE CODE WITHOUT THE EXPRESS WRITTEN CONSENT OF COMPANY IS
-// STRICTLY PROHIBITED, AND IN VIOLATION OF APPLICABLE LAWS AND
-// INTERNATIONAL TREATIES. THE RECEIPT OR POSSESSION OF THIS SOURCE
-// CODE AND/OR RELATED INFORMATION DOES NOT CONVEY OR IMPLY ANY RIGHTS
-// TO REPRODUCE, DISCLOSE OR DISTRIBUTE ITS CONTENTS, OR TO MANUFACTURE,
-// USE, OR SELL ANYTHING THAT IT MAY DESCRIBE, IN WHOLE OR IN PART.
-//
-// %COPYRIGHT_END%
-// --------------------------------------------------------------------
-// %BANNER_END%
+// Copyright 1998-2018 Epic Games, Inc. All Rights Reserved.
 
 #include "MagicLeapCustomPresent.h"
 #include "MagicLeapHMD.h"
diff --git a/Engine/Plugins/Lumin/MagicLeap/Source/MagicLeap/Private/MagicLeapCustomPresent.h b/Engine/Plugins/Lumin/MagicLeap/Source/MagicLeap/Private/MagicLeapCustomPresent.h
--- a/Engine/Plugins/Lumin/MagicLeap/Source/MagicLeap/Private/MagicLeapCustomPresent.h
+++ b/Engine/Plugins/Lumin/MagicLeap/Source/MagicLeap/Private/MagicLeapCustomPresent.h
@@ -1,36 +1,4 @@
-// %BANNER_BEGIN%
-// ---------------------------------------------------------------------
-// %COPYRIGHT_BEGIN%
-//
-// Copyright (c) 2017 Magic Leap, Inc. (COMPANY) All Rights Reserved.
-// Magic Leap, Inc. Confidential and Proprietary
-//
-// NOTICE: All information contained herein is, and remains the property
-// of COMPANY. The intellectual and technical concepts contained herein
-// are proprietary to COMPANY and may be covered by U.S. and Foreign
-// Patents, patents in process, and are protected by trade secret or
-// copyright law. Dissemination of this information or reproduction of
-// this material is strictly forbidden unless prior written permission is
-// obtained from COMPANY. Access to the source code contained herein is
-// hereby forbidden to anyone except current COMPANY employees, managers
-// or contractors who have executed Confidentiality and Non-disclosure
-// agreements explicitly covering such access.
-//
-// The copyright notice above does not evidence any actual or intended
-// publication or disclosure of this source code, which includes
-// information that is confidential and/or proprietary, and is a trade
-// secret, of COMPANY. ANY REPRODUCTION, MODIFICATION, DISTRIBUTION,
-// PUBLIC PERFORMANCE, OR PUBLIC DISPLAY OF OR THROUGH USE OF THIS
-// SOURCE CODE WITHOUT THE EXPRESS WRITTEN CONSENT OF COMPANY IS
-// STRICTLY PROHIBITED, AND IN VIOLATION OF APPLICABLE LAWS AND
-// INTERNATIONAL TREATIES. THE RECEIPT OR POSSESSION OF THIS SOURCE
-// CODE AND/OR RELATED INFORMATION DOES NOT CONVEY OR IMPLY ANY RIGHTS
-// TO REPRODUCE, DISCLOSE OR DISTRIBUTE ITS CONTENTS, OR TO MANUFACTURE,
-// USE, OR SELL ANYTHING THAT IT MAY DESCRIBE, IN WHOLE OR IN PART.
-//
-// %COPYRIGHT_END%
-// --------------------------------------------------------------------
-// %BANNER_END%
+// Copyright 1998-2018 Epic Games, Inc. All Rights Reserved.
 
 #pragma once
 
diff --git a/Engine/Plugins/Lumin/MagicLeap/Source/MagicLeap/Private/MagicLeapHMD.cpp b/Engine/Plugins/Lumin/MagicLeap/Source/MagicLeap/Private/MagicLeapHMD.cpp
--- a/Engine/Plugins/Lumin/MagicLeap/Source/MagicLeap/Private/MagicLeapHMD.cpp
+++ b/Engine/Plugins/Lumin/MagicLeap/Source/MagicLeap/Private/MagicLeapHMD.cpp
@@ -1,36 +1,4 @@
-// %BANNER_BEGIN%
-// ---------------------------------------------------------------------
-// %COPYRIGHT_BEGIN%
-//
-// Copyright (c) 2017 Magic Leap, Inc. (COMPANY) All Rights Reserved.
-// Magic Leap, Inc. Confidential and Proprietary
-//
-// NOTICE: All information contained herein is, and remains the property
-// of COMPANY. The intellectual and technical concepts contained herein
-// are proprietary to COMPANY and may be covered by U.S. and Foreign
-// Patents, patents in process, and are protected by trade secret or
-// copyright law. Dissemination of this information or reproduction of
-// this material is strictly forbidden unless prior written permission is
-// obtained from COMPANY. Access to the source code contained herein is
-// hereby forbidden to anyone except current COMPANY employees, managers
-// or contractors who have executed Confidentiality and Non-disclosure
-// agreements explicitly covering such access.
-//
-// The copyright notice above does not evidence any actual or intended
-// publication or disclosure of this source code, which includes
-// information that is confidential and/or proprietary, and is a trade
-// secret, of COMPANY. ANY REPRODUCTION, MODIFICATION, DISTRIBUTION,
-// PUBLIC PERFORMANCE, OR PUBLIC DISPLAY OF OR THROUGH USE OF THIS
-// SOURCE CODE WITHOUT THE EXPRESS WRITTEN CONSENT OF COMPANY IS
-// STRICTLY PROHIBITED, AND IN VIOLATION OF APPLICABLE LAWS AND
-// INTERNATIONAL TREATIES. THE RECEIPT OR POSSESSION OF THIS SOURCE
-// CODE AND/OR RELATED INFORMATION DOES NOT CONVEY OR IMPLY ANY RIGHTS
-// TO REPRODUCE, DISCLOSE OR DISTRIBUTE ITS CONTENTS, OR TO MANUFACTURE,
-// USE, OR SELL ANYTHING THAT IT MAY DESCRIBE, IN WHOLE OR IN PART.
-//
-// %COPYRIGHT_END%
-// --------------------------------------------------------------------
-// %BANNER_END%
+// Copyright 1998-2018 Epic Games, Inc. All Rights Reserved.
 
 #include "MagicLeapHMD.h"
 #include "MagicLeapHMDFunctionLibrary.h"
diff --git a/Engine/Plugins/Lumin/MagicLeap/Source/MagicLeap/Private/MagicLeapHMD.h b/Engine/Plugins/Lumin/MagicLeap/Source/MagicLeap/Private/MagicLeapHMD.h
--- a/Engine/Plugins/Lumin/MagicLeap/Source/MagicLeap/Private/MagicLeapHMD.h
+++ b/Engine/Plugins/Lumin/MagicLeap/Source/MagicLeap/Private/MagicLeapHMD.h
@@ -1,36 +1,4 @@
-// %BANNER_BEGIN%
-// ---------------------------------------------------------------------
-// %COPYRIGHT_BEGIN%
-//
-// Copyright (c) 2017 Magic Leap, Inc. (COMPANY) All Rights Reserved.
-// Magic Leap, Inc. Confidential and Proprietary
-//
-// NOTICE: All information contained herein is, and remains the property
-// of COMPANY. The intellectual and technical concepts contained herein
-// are proprietary to COMPANY and may be covered by U.S. and Foreign
-// Patents, patents in process, and are protected by trade secret or
-// copyright law. Dissemination of this information or reproduction of
-// this material is strictly forbidden unless prior written permission is
-// obtained from COMPANY. Access to the source code contained herein is
-// hereby forbidden to anyone except current COMPANY employees, managers
-// or contractors who have executed Confidentiality and Non-disclosure
-// agreements explicitly covering such access.
-//
-// The copyright notice above does not evidence any actual or intended
-// publication or disclosure of this source code, which includes
-// information that is confidential and/or proprietary, and is a trade
-// secret, of COMPANY. ANY REPRODUCTION, MODIFICATION, DISTRIBUTION,
-// PUBLIC PERFORMANCE, OR PUBLIC DISPLAY OF OR THROUGH USE OF THIS
-// SOURCE CODE WITHOUT THE EXPRESS WRITTEN CONSENT OF COMPANY IS
-// STRICTLY PROHIBITED, AND IN VIOLATION OF APPLICABLE LAWS AND
-// INTERNATIONAL TREATIES. THE RECEIPT OR POSSESSION OF THIS SOURCE
-// CODE AND/OR RELATED INFORMATION DOES NOT CONVEY OR IMPLY ANY RIGHTS
-// TO REPRODUCE, DISCLOSE OR DISTRIBUTE ITS CONTENTS, OR TO MANUFACTURE,
-// USE, OR SELL ANYTHING THAT IT MAY DESCRIBE, IN WHOLE OR IN PART.
-//
-// %COPYRIGHT_END%
-// --------------------------------------------------------------------
-// %BANNER_END%
+// Copyright 1998-2018 Epic Games, Inc. All Rights Reserved.
 
 #pragma once
 
diff --git a/Engine/Plugins/Lumin/MagicLeap/Source/MagicLeap/Private/MagicLeapHMDFunctionLibrary.cpp b/Engine/Plugins/Lumin/MagicLeap/Source/MagicLeap/Private/MagicLeapHMDFunctionLibrary.cpp
--- a/Engine/Plugins/Lumin/MagicLeap/Source/MagicLeap/Private/MagicLeapHMDFunctionLibrary.cpp
+++ b/Engine/Plugins/Lumin/MagicLeap/Source/MagicLeap/Private/MagicLeapHMDFunctionLibrary.cpp
@@ -1,36 +1,4 @@
-// %BANNER_BEGIN%
-// ---------------------------------------------------------------------
-// %COPYRIGHT_BEGIN%
-//
-// Copyright (c) 2017 Magic Leap, Inc. (COMPANY) All Rights Reserved.
-// Magic Leap, Inc. Confidential and Proprietary
-//
-// NOTICE: All information contained herein is, and remains the property
-// of COMPANY. The intellectual and technical concepts contained herein
-// are proprietary to COMPANY and may be covered by U.S. and Foreign
-// Patents, patents in process, and are protected by trade secret or
-// copyright law. Dissemination of this information or reproduction of
-// this material is strictly forbidden unless prior written permission is
-// obtained from COMPANY. Access to the source code contained herein is
-// hereby forbidden to anyone except current COMPANY employees, managers
-// or contractors who have executed Confidentiality and Non-disclosure
-// agreements explicitly covering such access.
-//
-// The copyright notice above does not evidence any actual or intended
-// publication or disclosure of this source code, which includes
-// information that is confidential and/or proprietary, and is a trade
-// secret, of COMPANY. ANY REPRODUCTION, MODIFICATION, DISTRIBUTION,
-// PUBLIC PERFORMANCE, OR PUBLIC DISPLAY OF OR THROUGH USE OF THIS
-// SOURCE CODE WITHOUT THE EXPRESS WRITTEN CONSENT OF COMPANY IS
-// STRICTLY PROHIBITED, AND IN VIOLATION OF APPLICABLE LAWS AND
-// INTERNATIONAL TREATIES. THE RECEIPT OR POSSESSION OF THIS SOURCE
-// CODE AND/OR RELATED INFORMATION DOES NOT CONVEY OR IMPLY ANY RIGHTS
-// TO REPRODUCE, DISCLOSE OR DISTRIBUTE ITS CONTENTS, OR TO MANUFACTURE,
-// USE, OR SELL ANYTHING THAT IT MAY DESCRIBE, IN WHOLE OR IN PART.
-//
-// %COPYRIGHT_END%
-// --------------------------------------------------------------------
-// %BANNER_END%
+// Copyright 1998-2018 Epic Games, Inc. All Rights Reserved.
 
 #include "MagicLeapHMDFunctionLibrary.h"
 #include "MagicLeapHMD.h"
diff --git a/Engine/Plugins/Lumin/MagicLeap/Source/MagicLeap/Private/MagicLeapSDKDetection.cpp b/Engine/Plugins/Lumin/MagicLeap/Source/MagicLeap/Private/MagicLeapSDKDetection.cpp
--- a/Engine/Plugins/Lumin/MagicLeap/Source/MagicLeap/Private/MagicLeapSDKDetection.cpp
+++ b/Engine/Plugins/Lumin/MagicLeap/Source/MagicLeap/Private/MagicLeapSDKDetection.cpp
@@ -1,36 +1,4 @@
-// %BANNER_BEGIN%
-// ---------------------------------------------------------------------
-// %COPYRIGHT_BEGIN%
-//
-// Copyright (c) 2017 Magic Leap, Inc. (COMPANY) All Rights Reserved.
-// Magic Leap, Inc. Confidential and Proprietary
-//
-// NOTICE: All information contained herein is, and remains the property
-// of COMPANY. The intellectual and technical concepts contained herein
-// are proprietary to COMPANY and may be covered by U.S. and Foreign
-// Patents, patents in process, and are protected by trade secret or
-// copyright law. Dissemination of this information or reproduction of
-// this material is strictly forbidden unless prior written permission is
-// obtained from COMPANY. Access to the source code contained herein is
-// hereby forbidden to anyone except current COMPANY employees, managers
-// or contractors who have executed Confidentiality and Non-disclosure
-// agreements explicitly covering such access.
-//
-// The copyright notice above does not evidence any actual or intended
-// publication or disclosure of this source code, which includes
-// information that is confidential and/or proprietary, and is a trade
-// secret, of COMPANY. ANY REPRODUCTION, MODIFICATION, DISTRIBUTION,
-// PUBLIC PERFORMANCE, OR PUBLIC DISPLAY OF OR THROUGH USE OF THIS
-// SOURCE CODE WITHOUT THE EXPRESS WRITTEN CONSENT OF COMPANY IS
-// STRICTLY PROHIBITED, AND IN VIOLATION OF APPLICABLE LAWS AND
-// INTERNATIONAL TREATIES. THE RECEIPT OR POSSESSION OF THIS SOURCE
-// CODE AND/OR RELATED INFORMATION DOES NOT CONVEY OR IMPLY ANY RIGHTS
-// TO REPRODUCE, DISCLOSE OR DISTRIBUTE ITS CONTENTS, OR TO MANUFACTURE,
-// USE, OR SELL ANYTHING THAT IT MAY DESCRIBE, IN WHOLE OR IN PART.
-//
-// %COPYRIGHT_END%
-// --------------------------------------------------------------------
-// %BANNER_END%
+// Copyright 1998-2018 Epic Games, Inc. All Rights Reserved.
 
 #include "MagicLeapSDKDetection.h"
 #include "MagicLeapPluginUtil.h"
diff --git a/Engine/Plugins/Lumin/MagicLeap/Source/MagicLeap/Private/MagicLeapUtils.cpp b/Engine/Plugins/Lumin/MagicLeap/Source/MagicLeap/Private/MagicLeapUtils.cpp
--- a/Engine/Plugins/Lumin/MagicLeap/Source/MagicLeap/Private/MagicLeapUtils.cpp
+++ b/Engine/Plugins/Lumin/MagicLeap/Source/MagicLeap/Private/MagicLeapUtils.cpp
@@ -1,36 +1,4 @@
-// %BANNER_BEGIN%
-// ---------------------------------------------------------------------
-// %COPYRIGHT_BEGIN%
-//
-// Copyright (c) 201x Magic Leap, Inc. (COMPANY) All Rights Reserved.
-// Magic Leap, Inc. Confidential and Proprietary
-//
-// NOTICE: All information contained herein is, and remains the property
-// of COMPANY. The intellectual and technical concepts contained herein
-// are proprietary to COMPANY and may be covered by U.S. and Foreign
-// Patents, patents in process, and are protected by trade secret or
-// copyright law. Dissemination of this information or reproduction of
-// this material is strictly forbidden unless prior written permission is
-// obtained from COMPANY. Access to the source code contained herein is
-// hereby forbidden to anyone except current COMPANY employees, managers
-// or contractors who have executed Confidentiality and Non-disclosure
-// agreements explicitly covering such access.
-//
-// The copyright notice above does not evidence any actual or intended
-// publication or disclosure of this source code, which includes
-// information that is confidential and/or proprietary, and is a trade
-// secret, of COMPANY. ANY REPRODUCTION, MODIFICATION, DISTRIBUTION,
-// PUBLIC PERFORMANCE, OR PUBLIC DISPLAY OF OR THROUGH USE OF THIS
-// SOURCE CODE WITHOUT THE EXPRESS WRITTEN CONSENT OF COMPANY IS
-// STRICTLY PROHIBITED, AND IN VIOLATION OF APPLICABLE LAWS AND
-// INTERNATIONAL TREATIES. THE RECEIPT OR POSSESSION OF THIS SOURCE
-// CODE AND/OR RELATED INFORMATION DOES NOT CONVEY OR IMPLY ANY RIGHTS
-// TO REPRODUCE, DISCLOSE OR DISTRIBUTE ITS CONTENTS, OR TO MANUFACTURE,
-// USE, OR SELL ANYTHING THAT IT MAY DESCRIBE, IN WHOLE OR IN PART.
-//
-// %COPYRIGHT_END%
-// --------------------------------------------------------------------
-// %BANNER_END%
+// Copyright 1998-2018 Epic Games, Inc. All Rights Reserved.
 
 #include "MagicLeapUtils.h"
 #include "MagicLeapMath.h"
diff --git a/Engine/Plugins/Lumin/MagicLeap/Source/MagicLeap/Private/MagicLeapUtils.h b/Engine/Plugins/Lumin/MagicLeap/Source/MagicLeap/Private/MagicLeapUtils.h
--- a/Engine/Plugins/Lumin/MagicLeap/Source/MagicLeap/Private/MagicLeapUtils.h
+++ b/Engine/Plugins/Lumin/MagicLeap/Source/MagicLeap/Private/MagicLeapUtils.h
@@ -1,36 +1,4 @@
-// %BANNER_BEGIN%
-// ---------------------------------------------------------------------
-// %COPYRIGHT_BEGIN%
-//
-// Copyright (c) 201x Magic Leap, Inc. (COMPANY) All Rights Reserved.
-// Magic Leap, Inc. Confidential and Proprietary
-//
-// NOTICE: All information contained herein is, and remains the property
-// of COMPANY. The intellectual and technical concepts contained herein
-// are proprietary to COMPANY and may be covered by U.S. and Foreign
-// Patents, patents in process, and are protected by trade secret or
-// copyright law. Dissemination of this information or reproduction of
-// this material is strictly forbidden unless prior written permission is
-// obtained from COMPANY. Access to the source code contained herein is
-// hereby forbidden to anyone except current COMPANY employees, managers
-// or contractors who have executed Confidentiality and Non-disclosure
-// agreements explicitly covering such access.
-//
-// The copyright notice above does not evidence any actual or intended
-// publication or disclosure of this source code, which includes
-// information that is confidential and/or proprietary, and is a trade
-// secret, of COMPANY. ANY REPRODUCTION, MODIFICATION, DISTRIBUTION,
-// PUBLIC PERFORMANCE, OR PUBLIC DISPLAY OF OR THROUGH USE OF THIS
-// SOURCE CODE WITHOUT THE EXPRESS WRITTEN CONSENT OF COMPANY IS
-// STRICTLY PROHIBITED, AND IN VIOLATION OF APPLICABLE LAWS AND
-// INTERNATIONAL TREATIES. THE RECEIPT OR POSSESSION OF THIS SOURCE
-// CODE AND/OR RELATED INFORMATION DOES NOT CONVEY OR IMPLY ANY RIGHTS
-// TO REPRODUCE, DISCLOSE OR DISTRIBUTE ITS CONTENTS, OR TO MANUFACTURE,
-// USE, OR SELL ANYTHING THAT IT MAY DESCRIBE, IN WHOLE OR IN PART.
-//
-// %COPYRIGHT_END%
-// --------------------------------------------------------------------
-// %BANNER_END%
+// Copyright 1998-2018 Epic Games, Inc. All Rights Reserved.
 
 #pragma once
 
diff --git a/Engine/Plugins/Lumin/MagicLeap/Source/MagicLeap/Private/MagicLeapXRCamera.cpp b/Engine/Plugins/Lumin/MagicLeap/Source/MagicLeap/Private/MagicLeapXRCamera.cpp
--- a/Engine/Plugins/Lumin/MagicLeap/Source/MagicLeap/Private/MagicLeapXRCamera.cpp
+++ b/Engine/Plugins/Lumin/MagicLeap/Source/MagicLeap/Private/MagicLeapXRCamera.cpp
@@ -1,36 +1,4 @@
-// %BANNER_BEGIN%
-// ---------------------------------------------------------------------
-// %COPYRIGHT_BEGIN%
-//
-// Copyright (c) 2017 Magic Leap, Inc. (COMPANY) All Rights Reserved.
-// Magic Leap, Inc. Confidential and Proprietary
-//
-// NOTICE: All information contained herein is, and remains the property
-// of COMPANY. The intellectual and technical concepts contained herein
-// are proprietary to COMPANY and may be covered by U.S. and Foreign
-// Patents, patents in process, and are protected by trade secret or
-// copyright law. Dissemination of this information or reproduction of
-// this material is strictly forbidden unless prior written permission is
-// obtained from COMPANY. Access to the source code contained herein is
-// hereby forbidden to anyone except current COMPANY employees, managers
-// or contractors who have executed Confidentiality and Non-disclosure
-// agreements explicitly covering such access.
-//
-// The copyright notice above does not evidence any actual or intended
-// publication or disclosure of this source code, which includes
-// information that is confidential and/or proprietary, and is a trade
-// secret, of COMPANY. ANY REPRODUCTION, MODIFICATION, DISTRIBUTION,
-// PUBLIC PERFORMANCE, OR PUBLIC DISPLAY OF OR THROUGH USE OF THIS
-// SOURCE CODE WITHOUT THE EXPRESS WRITTEN CONSENT OF COMPANY IS
-// STRICTLY PROHIBITED, AND IN VIOLATION OF APPLICABLE LAWS AND
-// INTERNATIONAL TREATIES. THE RECEIPT OR POSSESSION OF THIS SOURCE
-// CODE AND/OR RELATED INFORMATION DOES NOT CONVEY OR IMPLY ANY RIGHTS
-// TO REPRODUCE, DISCLOSE OR DISTRIBUTE ITS CONTENTS, OR TO MANUFACTURE,
-// USE, OR SELL ANYTHING THAT IT MAY DESCRIBE, IN WHOLE OR IN PART.
-//
-// %COPYRIGHT_END%
-// --------------------------------------------------------------------
-// %BANNER_END%
+// Copyright 1998-2018 Epic Games, Inc. All Rights Reserved.
 
 #include "MagicLeapXRCamera.h"
 #include "SceneView.h"
diff --git a/Engine/Plugins/Lumin/MagicLeap/Source/MagicLeap/Private/MagicLeapXRCamera.h b/Engine/Plugins/Lumin/MagicLeap/Source/MagicLeap/Private/MagicLeapXRCamera.h
--- a/Engine/Plugins/Lumin/MagicLeap/Source/MagicLeap/Private/MagicLeapXRCamera.h
+++ b/Engine/Plugins/Lumin/MagicLeap/Source/MagicLeap/Private/MagicLeapXRCamera.h
@@ -1,36 +1,4 @@
-// %BANNER_BEGIN%
-// ---------------------------------------------------------------------
-// %COPYRIGHT_BEGIN%
-//
-// Copyright (c) 2017 Magic Leap, Inc. (COMPANY) All Rights Reserved.
-// Magic Leap, Inc. Confidential and Proprietary
-//
-// NOTICE: All information contained herein is, and remains the property
-// of COMPANY. The intellectual and technical concepts contained herein
-// are proprietary to COMPANY and may be covered by U.S. and Foreign
-// Patents, patents in process, and are protected by trade secret or
-// copyright law. Dissemination of this information or reproduction of
-// this material is strictly forbidden unless prior written permission is
-// obtained from COMPANY. Access to the source code contained herein is
-// hereby forbidden to anyone except current COMPANY employees, managers
-// or contractors who have executed Confidentiality and Non-disclosure
-// agreements explicitly covering such access.
-//
-// The copyright notice above does not evidence any actual or intended
-// publication or disclosure of this source code, which includes
-// information that is confidential and/or proprietary, and is a trade
-// secret, of COMPANY. ANY REPRODUCTION, MODIFICATION, DISTRIBUTION,
-// PUBLIC PERFORMANCE, OR PUBLIC DISPLAY OF OR THROUGH USE OF THIS
-// SOURCE CODE WITHOUT THE EXPRESS WRITTEN CONSENT OF COMPANY IS
-// STRICTLY PROHIBITED, AND IN VIOLATION OF APPLICABLE LAWS AND
-// INTERNATIONAL TREATIES. THE RECEIPT OR POSSESSION OF THIS SOURCE
-// CODE AND/OR RELATED INFORMATION DOES NOT CONVEY OR IMPLY ANY RIGHTS
-// TO REPRODUCE, DISCLOSE OR DISTRIBUTE ITS CONTENTS, OR TO MANUFACTURE,
-// USE, OR SELL ANYTHING THAT IT MAY DESCRIBE, IN WHOLE OR IN PART.
-//
-// %COPYRIGHT_END%
-// --------------------------------------------------------------------
-// %BANNER_END%
+// Copyright 1998-2018 Epic Games, Inc. All Rights Reserved.
 
 #pragma once
 
diff --git a/Engine/Plugins/Lumin/MagicLeap/Source/MagicLeap/Private/MeshTrackerComponent.cpp b/Engine/Plugins/Lumin/MagicLeap/Source/MagicLeap/Private/MeshTrackerComponent.cpp
--- a/Engine/Plugins/Lumin/MagicLeap/Source/MagicLeap/Private/MeshTrackerComponent.cpp
+++ b/Engine/Plugins/Lumin/MagicLeap/Source/MagicLeap/Private/MeshTrackerComponent.cpp
@@ -1,36 +1,4 @@
-// %BANNER_BEGIN%
-// ---------------------------------------------------------------------
-// %COPYRIGHT_BEGIN%
-//
-// Copyright (c) 2017 Magic Leap, Inc. (COMPANY) All Rights Reserved.
-// Magic Leap, Inc. Confidential and Proprietary
-//
-// NOTICE: All information contained herein is, and remains the property
-// of COMPANY. The intellectual and technical concepts contained herein
-// are proprietary to COMPANY and may be covered by U.S. and Foreign
-// Patents, patents in process, and are protected by trade secret or
-// copyright law. Dissemination of this information or reproduction of
-// this material is strictly forbidden unless prior written permission is
-// obtained from COMPANY. Access to the source code contained herein is
-// hereby forbidden to anyone except current COMPANY employees, managers
-// or contractors who have executed Confidentiality and Non-disclosure
-// agreements explicitly covering such access.
-//
-// The copyright notice above does not evidence any actual or intended
-// publication or disclosure of this source code, which includes
-// information that is confidential and/or proprietary, and is a trade
-// secret, of COMPANY. ANY REPRODUCTION, MODIFICATION, DISTRIBUTION,
-// PUBLIC PERFORMANCE, OR PUBLIC DISPLAY OF OR THROUGH USE OF THIS
-// SOURCE CODE WITHOUT THE EXPRESS WRITTEN CONSENT OF COMPANY IS
-// STRICTLY PROHIBITED, AND IN VIOLATION OF APPLICABLE LAWS AND
-// INTERNATIONAL TREATIES. THE RECEIPT OR POSSESSION OF THIS SOURCE
-// CODE AND/OR RELATED INFORMATION DOES NOT CONVEY OR IMPLY ANY RIGHTS
-// TO REPRODUCE, DISCLOSE OR DISTRIBUTE ITS CONTENTS, OR TO MANUFACTURE,
-// USE, OR SELL ANYTHING THAT IT MAY DESCRIBE, IN WHOLE OR IN PART.
-//
-// %COPYRIGHT_END%
-// --------------------------------------------------------------------
-// %BANNER_END%
+// Copyright 1998-2018 Epic Games, Inc. All Rights Reserved.
 
 #include "MeshTrackerComponent.h"
 #include "MagicLeapHMD.h"
diff --git a/Engine/Plugins/Lumin/MagicLeap/Source/MagicLeap/Private/PlanesComponent.cpp b/Engine/Plugins/Lumin/MagicLeap/Source/MagicLeap/Private/PlanesComponent.cpp
--- a/Engine/Plugins/Lumin/MagicLeap/Source/MagicLeap/Private/PlanesComponent.cpp
+++ b/Engine/Plugins/Lumin/MagicLeap/Source/MagicLeap/Private/PlanesComponent.cpp
@@ -1,36 +1,4 @@
-// %BANNER_BEGIN%
-// ---------------------------------------------------------------------
-// %COPYRIGHT_BEGIN%
-//
-// Copyright (c) 2017 Magic Leap, Inc. (COMPANY) All Rights Reserved.
-// Magic Leap, Inc. Confidential and Proprietary
-//
-// NOTICE: All information contained herein is, and remains the property
-// of COMPANY. The intellectual and technical concepts contained herein
-// are proprietary to COMPANY and may be covered by U.S. and Foreign
-// Patents, patents in process, and are protected by trade secret or
-// copyright law. Dissemination of this information or reproduction of
-// this material is strictly forbidden unless prior written permission is
-// obtained from COMPANY. Access to the source code contained herein is
-// hereby forbidden to anyone except current COMPANY employees, managers
-// or contractors who have executed Confidentiality and Non-disclosure
-// agreements explicitly covering such access.
-//
-// The copyright notice above does not evidence any actual or intended
-// publication or disclosure of this source code, which includes
-// information that is confidential and/or proprietary, and is a trade
-// secret, of COMPANY. ANY REPRODUCTION, MODIFICATION, DISTRIBUTION,
-// PUBLIC PERFORMANCE, OR PUBLIC DISPLAY OF OR THROUGH USE OF THIS
-// SOURCE CODE WITHOUT THE EXPRESS WRITTEN CONSENT OF COMPANY IS
-// STRICTLY PROHIBITED, AND IN VIOLATION OF APPLICABLE LAWS AND
-// INTERNATIONAL TREATIES. THE RECEIPT OR POSSESSION OF THIS SOURCE
-// CODE AND/OR RELATED INFORMATION DOES NOT CONVEY OR IMPLY ANY RIGHTS
-// TO REPRODUCE, DISCLOSE OR DISTRIBUTE ITS CONTENTS, OR TO MANUFACTURE,
-// USE, OR SELL ANYTHING THAT IT MAY DESCRIBE, IN WHOLE OR IN PART.
-//
-// %COPYRIGHT_END%
-// --------------------------------------------------------------------
-// %BANNER_END%
+// Copyright 1998-2018 Epic Games, Inc. All Rights Reserved.
 
 #include "PlanesComponent.h"
 #include "Components/BoxComponent.h"
diff --git a/Engine/Plugins/Lumin/MagicLeap/Source/MagicLeap/Private/RaycastComponent.cpp b/Engine/Plugins/Lumin/MagicLeap/Source/MagicLeap/Private/RaycastComponent.cpp
--- a/Engine/Plugins/Lumin/MagicLeap/Source/MagicLeap/Private/RaycastComponent.cpp
+++ b/Engine/Plugins/Lumin/MagicLeap/Source/MagicLeap/Private/RaycastComponent.cpp
@@ -1,36 +1,4 @@
-// %BANNER_BEGIN%
-// ---------------------------------------------------------------------
-// %COPYRIGHT_BEGIN%
-//
-// Copyright (c) 2017 Magic Leap, Inc. (COMPANY) All Rights Reserved.
-// Magic Leap, Inc. Confidential and Proprietary
-//
-// NOTICE: All information contained herein is, and remains the property
-// of COMPANY. The intellectual and technical concepts contained herein
-// are proprietary to COMPANY and may be covered by U.S. and Foreign
-// Patents, patents in process, and are protected by trade secret or
-// copyright law. Dissemination of this information or reproduction of
-// this material is strictly forbidden unless prior written permission is
-// obtained from COMPANY. Access to the source code contained herein is
-// hereby forbidden to anyone except current COMPANY employees, managers
-// or contractors who have executed Confidentiality and Non-disclosure
-// agreements explicitly covering such access.
-//
-// The copyright notice above does not evidence any actual or intended
-// publication or disclosure of this source code, which includes
-// information that is confidential and/or proprietary, and is a trade
-// secret, of COMPANY. ANY REPRODUCTION, MODIFICATION, DISTRIBUTION,
-// PUBLIC PERFORMANCE, OR PUBLIC DISPLAY OF OR THROUGH USE OF THIS
-// SOURCE CODE WITHOUT THE EXPRESS WRITTEN CONSENT OF COMPANY IS
-// STRICTLY PROHIBITED, AND IN VIOLATION OF APPLICABLE LAWS AND
-// INTERNATIONAL TREATIES. THE RECEIPT OR POSSESSION OF THIS SOURCE
-// CODE AND/OR RELATED INFORMATION DOES NOT CONVEY OR IMPLY ANY RIGHTS
-// TO REPRODUCE, DISCLOSE OR DISTRIBUTE ITS CONTENTS, OR TO MANUFACTURE,
-// USE, OR SELL ANYTHING THAT IT MAY DESCRIBE, IN WHOLE OR IN PART.
-//
-// %COPYRIGHT_END%
-// --------------------------------------------------------------------
-// %BANNER_END%
+// Copyright 1998-2018 Epic Games, Inc. All Rights Reserved.
 
 #include "RaycastComponent.h"
 #include "MagicLeapHMD.h"
diff --git a/Engine/Plugins/Lumin/MagicLeap/Source/MagicLeap/Public/IMagicLeapHMD.h b/Engine/Plugins/Lumin/MagicLeap/Source/MagicLeap/Public/IMagicLeapHMD.h
--- a/Engine/Plugins/Lumin/MagicLeap/Source/MagicLeap/Public/IMagicLeapHMD.h
+++ b/Engine/Plugins/Lumin/MagicLeap/Source/MagicLeap/Public/IMagicLeapHMD.h
@@ -1,36 +1,4 @@
-// %BANNER_BEGIN%
-// ---------------------------------------------------------------------
-// %COPYRIGHT_BEGIN%
-//
-// Copyright (c) 2017 Magic Leap, Inc. (COMPANY) All Rights Reserved.
-// Magic Leap, Inc. Confidential and Proprietary
-//
-// NOTICE: All information contained herein is, and remains the property
-// of COMPANY. The intellectual and technical concepts contained herein
-// are proprietary to COMPANY and may be covered by U.S. and Foreign
-// Patents, patents in process, and are protected by trade secret or
-// copyright law. Dissemination of this information or reproduction of
-// this material is strictly forbidden unless prior written permission is
-// obtained from COMPANY. Access to the source code contained herein is
-// hereby forbidden to anyone except current COMPANY employees, managers
-// or contractors who have executed Confidentiality and Non-disclosure
-// agreements explicitly covering such access.
-//
-// The copyright notice above does not evidence any actual or intended
-// publication or disclosure of this source code, which includes
-// information that is confidential and/or proprietary, and is a trade
-// secret, of COMPANY. ANY REPRODUCTION, MODIFICATION, DISTRIBUTION,
-// PUBLIC PERFORMANCE, OR PUBLIC DISPLAY OF OR THROUGH USE OF THIS
-// SOURCE CODE WITHOUT THE EXPRESS WRITTEN CONSENT OF COMPANY IS
-// STRICTLY PROHIBITED, AND IN VIOLATION OF APPLICABLE LAWS AND
-// INTERNATIONAL TREATIES. THE RECEIPT OR POSSESSION OF THIS SOURCE
-// CODE AND/OR RELATED INFORMATION DOES NOT CONVEY OR IMPLY ANY RIGHTS
-// TO REPRODUCE, DISCLOSE OR DISTRIBUTE ITS CONTENTS, OR TO MANUFACTURE,
-// USE, OR SELL ANYTHING THAT IT MAY DESCRIBE, IN WHOLE OR IN PART.
-//
-// %COPYRIGHT_END%
-// --------------------------------------------------------------------
-// %BANNER_END%
+// Copyright 1998-2018 Epic Games, Inc. All Rights Reserved.
 
 #pragma once
 
diff --git a/Engine/Plugins/Lumin/MagicLeap/Source/MagicLeap/Public/IMagicLeapInputDevice.h b/Engine/Plugins/Lumin/MagicLeap/Source/MagicLeap/Public/IMagicLeapInputDevice.h
--- a/Engine/Plugins/Lumin/MagicLeap/Source/MagicLeap/Public/IMagicLeapInputDevice.h
+++ b/Engine/Plugins/Lumin/MagicLeap/Source/MagicLeap/Public/IMagicLeapInputDevice.h
@@ -1,36 +1,4 @@
-// %BANNER_BEGIN%
-// ---------------------------------------------------------------------
-// %COPYRIGHT_BEGIN%
-//
-// Copyright (c) 2017 Magic Leap, Inc. (COMPANY) All Rights Reserved.
-// Magic Leap, Inc. Confidential and Proprietary
-//
-// NOTICE: All information contained herein is, and remains the property
-// of COMPANY. The intellectual and technical concepts contained herein
-// are proprietary to COMPANY and may be covered by U.S. and Foreign
-// Patents, patents in process, and are protected by trade secret or
-// copyright law. Dissemination of this information or reproduction of
-// this material is strictly forbidden unless prior written permission is
-// obtained from COMPANY. Access to the source code contained herein is
-// hereby forbidden to anyone except current COMPANY employees, managers
-// or contractors who have executed Confidentiality and Non-disclosure
-// agreements explicitly covering such access.
-//
-// The copyright notice above does not evidence any actual or intended
-// publication or disclosure of this source code, which includes
-// information that is confidential and/or proprietary, and is a trade
-// secret, of COMPANY. ANY REPRODUCTION, MODIFICATION, DISTRIBUTION,
-// PUBLIC PERFORMANCE, OR PUBLIC DISPLAY OF OR THROUGH USE OF THIS
-// SOURCE CODE WITHOUT THE EXPRESS WRITTEN CONSENT OF COMPANY IS
-// STRICTLY PROHIBITED, AND IN VIOLATION OF APPLICABLE LAWS AND
-// INTERNATIONAL TREATIES. THE RECEIPT OR POSSESSION OF THIS SOURCE
-// CODE AND/OR RELATED INFORMATION DOES NOT CONVEY OR IMPLY ANY RIGHTS
-// TO REPRODUCE, DISCLOSE OR DISTRIBUTE ITS CONTENTS, OR TO MANUFACTURE,
-// USE, OR SELL ANYTHING THAT IT MAY DESCRIBE, IN WHOLE OR IN PART.
-//
-// %COPYRIGHT_END%
-// --------------------------------------------------------------------
-// %BANNER_END%
+// Copyright 1998-2018 Epic Games, Inc. All Rights Reserved.
 
 #pragma once
 
diff --git a/Engine/Plugins/Lumin/MagicLeap/Source/MagicLeap/Public/IMagicLeapPlugin.h b/Engine/Plugins/Lumin/MagicLeap/Source/MagicLeap/Public/IMagicLeapPlugin.h
--- a/Engine/Plugins/Lumin/MagicLeap/Source/MagicLeap/Public/IMagicLeapPlugin.h
+++ b/Engine/Plugins/Lumin/MagicLeap/Source/MagicLeap/Public/IMagicLeapPlugin.h
@@ -1,36 +1,4 @@
-// %BANNER_BEGIN%
-// ---------------------------------------------------------------------
-// %COPYRIGHT_BEGIN%
-//
-// Copyright (c) 2017 Magic Leap, Inc. (COMPANY) All Rights Reserved.
-// Magic Leap, Inc. Confidential and Proprietary
-//
-// NOTICE: All information contained herein is, and remains the property
-// of COMPANY. The intellectual and technical concepts contained herein
-// are proprietary to COMPANY and may be covered by U.S. and Foreign
-// Patents, patents in process, and are protected by trade secret or
-// copyright law. Dissemination of this information or reproduction of
-// this material is strictly forbidden unless prior written permission is
-// obtained from COMPANY. Access to the source code contained herein is
-// hereby forbidden to anyone except current COMPANY employees, managers
-// or contractors who have executed Confidentiality and Non-disclosure
-// agreements explicitly covering such access.
-//
-// The copyright notice above does not evidence any actual or intended
-// publication or disclosure of this source code, which includes
-// information that is confidential and/or proprietary, and is a trade
-// secret, of COMPANY. ANY REPRODUCTION, MODIFICATION, DISTRIBUTION,
-// PUBLIC PERFORMANCE, OR PUBLIC DISPLAY OF OR THROUGH USE OF THIS
-// SOURCE CODE WITHOUT THE EXPRESS WRITTEN CONSENT OF COMPANY IS
-// STRICTLY PROHIBITED, AND IN VIOLATION OF APPLICABLE LAWS AND
-// INTERNATIONAL TREATIES. THE RECEIPT OR POSSESSION OF THIS SOURCE
-// CODE AND/OR RELATED INFORMATION DOES NOT CONVEY OR IMPLY ANY RIGHTS
-// TO REPRODUCE, DISCLOSE OR DISTRIBUTE ITS CONTENTS, OR TO MANUFACTURE,
-// USE, OR SELL ANYTHING THAT IT MAY DESCRIBE, IN WHOLE OR IN PART.
-//
-// %COPYRIGHT_END%
-// --------------------------------------------------------------------
-// %BANNER_END%
+// Copyright 1998-2018 Epic Games, Inc. All Rights Reserved.
 
 #pragma once
 
diff --git a/Engine/Plugins/Lumin/MagicLeap/Source/MagicLeap/Public/MagicLeapSDKDetection.h b/Engine/Plugins/Lumin/MagicLeap/Source/MagicLeap/Public/MagicLeapSDKDetection.h
--- a/Engine/Plugins/Lumin/MagicLeap/Source/MagicLeap/Public/MagicLeapSDKDetection.h
+++ b/Engine/Plugins/Lumin/MagicLeap/Source/MagicLeap/Public/MagicLeapSDKDetection.h
@@ -1,36 +1,4 @@
-// %BANNER_BEGIN%
-// ---------------------------------------------------------------------
-// %COPYRIGHT_BEGIN%
-//
-// Copyright (c) 2017 Magic Leap, Inc. (COMPANY) All Rights Reserved.
-// Magic Leap, Inc. Confidential and Proprietary
-//
-// NOTICE: All information contained herein is, and remains the property
-// of COMPANY. The intellectual and technical concepts contained herein
-// are proprietary to COMPANY and may be covered by U.S. and Foreign
-// Patents, patents in process, and are protected by trade secret or
-// copyright law. Dissemination of this information or reproduction of
-// this material is strictly forbidden unless prior written permission is
-// obtained from COMPANY. Access to the source code contained herein is
-// hereby forbidden to anyone except current COMPANY employees, managers
-// or contractors who have executed Confidentiality and Non-disclosure
-// agreements explicitly covering such access.
-//
-// The copyright notice above does not evidence any actual or intended
-// publication or disclosure of this source code, which includes
-// information that is confidential and/or proprietary, and is a trade
-// secret, of COMPANY. ANY REPRODUCTION, MODIFICATION, DISTRIBUTION,
-// PUBLIC PERFORMANCE, OR PUBLIC DISPLAY OF OR THROUGH USE OF THIS
-// SOURCE CODE WITHOUT THE EXPRESS WRITTEN CONSENT OF COMPANY IS
-// STRICTLY PROHIBITED, AND IN VIOLATION OF APPLICABLE LAWS AND
-// INTERNATIONAL TREATIES. THE RECEIPT OR POSSESSION OF THIS SOURCE
-// CODE AND/OR RELATED INFORMATION DOES NOT CONVEY OR IMPLY ANY RIGHTS
-// TO REPRODUCE, DISCLOSE OR DISTRIBUTE ITS CONTENTS, OR TO MANUFACTURE,
-// USE, OR SELL ANYTHING THAT IT MAY DESCRIBE, IN WHOLE OR IN PART.
-//
-// %COPYRIGHT_END%
-// --------------------------------------------------------------------
-// %BANNER_END%
+// Copyright 1998-2018 Epic Games, Inc. All Rights Reserved.
 
 #pragma once
 
diff --git a/Engine/Plugins/Lumin/MagicLeap/Source/MagicLeapController/MagicLeapController.Build.cs b/Engine/Plugins/Lumin/MagicLeap/Source/MagicLeapController/MagicLeapController.Build.cs
--- a/Engine/Plugins/Lumin/MagicLeap/Source/MagicLeapController/MagicLeapController.Build.cs
+++ b/Engine/Plugins/Lumin/MagicLeap/Source/MagicLeapController/MagicLeapController.Build.cs
@@ -1,36 +1,4 @@
-// %BANNER_BEGIN%
-// ---------------------------------------------------------------------
-// %COPYRIGHT_BEGIN%
-//
-// Copyright (c) 2017 Magic Leap, Inc. (COMPANY) All Rights Reserved.
-// Magic Leap, Inc. Confidential and Proprietary
-//
-// NOTICE: All information contained herein is, and remains the property
-// of COMPANY. The intellectual and technical concepts contained herein
-// are proprietary to COMPANY and may be covered by U.S. and Foreign
-// Patents, patents in process, and are protected by trade secret or
-// copyright law. Dissemination of this information or reproduction of
-// this material is strictly forbidden unless prior written permission is
-// obtained from COMPANY. Access to the source code contained herein is
-// hereby forbidden to anyone except current COMPANY employees, managers
-// or contractors who have executed Confidentiality and Non-disclosure
-// agreements explicitly covering such access.
-//
-// The copyright notice above does not evidence any actual or intended
-// publication or disclosure of this source code, which includes
-// information that is confidential and/or proprietary, and is a trade
-// secret, of COMPANY. ANY REPRODUCTION, MODIFICATION, DISTRIBUTION,
-// PUBLIC PERFORMANCE, OR PUBLIC DISPLAY OF OR THROUGH USE OF THIS
-// SOURCE CODE WITHOUT THE EXPRESS WRITTEN CONSENT OF COMPANY IS
-// STRICTLY PROHIBITED, AND IN VIOLATION OF APPLICABLE LAWS AND
-// INTERNATIONAL TREATIES. THE RECEIPT OR POSSESSION OF THIS SOURCE
-// CODE AND/OR RELATED INFORMATION DOES NOT CONVEY OR IMPLY ANY RIGHTS
-// TO REPRODUCE, DISCLOSE OR DISTRIBUTE ITS CONTENTS, OR TO MANUFACTURE,
-// USE, OR SELL ANYTHING THAT IT MAY DESCRIBE, IN WHOLE OR IN PART.
-//
-// %COPYRIGHT_END%
-// --------------------------------------------------------------------
-// %BANNER_END%
+// Copyright 1998-2018 Epic Games, Inc. All Rights Reserved.
 
 using System.IO;
 
diff --git a/Engine/Plugins/Lumin/MagicLeap/Source/MagicLeapController/Private/MagicLeapController.cpp b/Engine/Plugins/Lumin/MagicLeap/Source/MagicLeapController/Private/MagicLeapController.cpp
--- a/Engine/Plugins/Lumin/MagicLeap/Source/MagicLeapController/Private/MagicLeapController.cpp
+++ b/Engine/Plugins/Lumin/MagicLeap/Source/MagicLeapController/Private/MagicLeapController.cpp
@@ -1,36 +1,4 @@
-// %BANNER_BEGIN%
-// ---------------------------------------------------------------------
-// %COPYRIGHT_BEGIN%
-//
-// Copyright (c) 2017 Magic Leap, Inc. (COMPANY) All Rights Reserved.
-// Magic Leap, Inc. Confidential and Proprietary
-//
-// NOTICE: All information contained herein is, and remains the property
-// of COMPANY. The intellectual and technical concepts contained herein
-// are proprietary to COMPANY and may be covered by U.S. and Foreign
-// Patents, patents in process, and are protected by trade secret or
-// copyright law. Dissemination of this information or reproduction of
-// this material is strictly forbidden unless prior written permission is
-// obtained from COMPANY. Access to the source code contained herein is
-// hereby forbidden to anyone except current COMPANY employees, managers
-// or contractors who have executed Confidentiality and Non-disclosure
-// agreements explicitly covering such access.
-//
-// The copyright notice above does not evidence any actual or intended
-// publication or disclosure of this source code, which includes
-// information that is confidential and/or proprietary, and is a trade
-// secret, of COMPANY. ANY REPRODUCTION, MODIFICATION, DISTRIBUTION,
-// PUBLIC PERFORMANCE, OR PUBLIC DISPLAY OF OR THROUGH USE OF THIS
-// SOURCE CODE WITHOUT THE EXPRESS WRITTEN CONSENT OF COMPANY IS
-// STRICTLY PROHIBITED, AND IN VIOLATION OF APPLICABLE LAWS AND
-// INTERNATIONAL TREATIES. THE RECEIPT OR POSSESSION OF THIS SOURCE
-// CODE AND/OR RELATED INFORMATION DOES NOT CONVEY OR IMPLY ANY RIGHTS
-// TO REPRODUCE, DISCLOSE OR DISTRIBUTE ITS CONTENTS, OR TO MANUFACTURE,
-// USE, OR SELL ANYTHING THAT IT MAY DESCRIBE, IN WHOLE OR IN PART.
-//
-// %COPYRIGHT_END%
-// --------------------------------------------------------------------
-// %BANNER_END%
+// Copyright 1998-2018 Epic Games, Inc. All Rights Reserved.
 
 #include "MagicLeapController.h"
 #include "IMagicLeapControllerPlugin.h"
diff --git a/Engine/Plugins/Lumin/MagicLeap/Source/MagicLeapController/Private/MagicLeapController.h b/Engine/Plugins/Lumin/MagicLeap/Source/MagicLeapController/Private/MagicLeapController.h
--- a/Engine/Plugins/Lumin/MagicLeap/Source/MagicLeapController/Private/MagicLeapController.h
+++ b/Engine/Plugins/Lumin/MagicLeap/Source/MagicLeapController/Private/MagicLeapController.h
@@ -1,36 +1,4 @@
-// %BANNER_BEGIN%
-// ---------------------------------------------------------------------
-// %COPYRIGHT_BEGIN%
-//
-// Copyright (c) 2017 Magic Leap, Inc. (COMPANY) All Rights Reserved.
-// Magic Leap, Inc. Confidential and Proprietary
-//
-// NOTICE: All information contained herein is, and remains the property
-// of COMPANY. The intellectual and technical concepts contained herein
-// are proprietary to COMPANY and may be covered by U.S. and Foreign
-// Patents, patents in process, and are protected by trade secret or
-// copyright law. Dissemination of this information or reproduction of
-// this material is strictly forbidden unless prior written permission is
-// obtained from COMPANY. Access to the source code contained herein is
-// hereby forbidden to anyone except current COMPANY employees, managers
-// or contractors who have executed Confidentiality and Non-disclosure
-// agreements explicitly covering such access.
-//
-// The copyright notice above does not evidence any actual or intended
-// publication or disclosure of this source code, which includes
-// information that is confidential and/or proprietary, and is a trade
-// secret, of COMPANY. ANY REPRODUCTION, MODIFICATION, DISTRIBUTION,
-// PUBLIC PERFORMANCE, OR PUBLIC DISPLAY OF OR THROUGH USE OF THIS
-// SOURCE CODE WITHOUT THE EXPRESS WRITTEN CONSENT OF COMPANY IS
-// STRICTLY PROHIBITED, AND IN VIOLATION OF APPLICABLE LAWS AND
-// INTERNATIONAL TREATIES. THE RECEIPT OR POSSESSION OF THIS SOURCE
-// CODE AND/OR RELATED INFORMATION DOES NOT CONVEY OR IMPLY ANY RIGHTS
-// TO REPRODUCE, DISCLOSE OR DISTRIBUTE ITS CONTENTS, OR TO MANUFACTURE,
-// USE, OR SELL ANYTHING THAT IT MAY DESCRIBE, IN WHOLE OR IN PART.
-//
-// %COPYRIGHT_END%
-// --------------------------------------------------------------------
-// %BANNER_END%
+// Copyright 1998-2018 Epic Games, Inc. All Rights Reserved.
 
 #pragma once
 
diff --git a/Engine/Plugins/Lumin/MagicLeap/Source/MagicLeapController/Private/MagicLeapControllerFunctionLibrary.cpp b/Engine/Plugins/Lumin/MagicLeap/Source/MagicLeapController/Private/MagicLeapControllerFunctionLibrary.cpp
--- a/Engine/Plugins/Lumin/MagicLeap/Source/MagicLeapController/Private/MagicLeapControllerFunctionLibrary.cpp
+++ b/Engine/Plugins/Lumin/MagicLeap/Source/MagicLeapController/Private/MagicLeapControllerFunctionLibrary.cpp
@@ -1,36 +1,4 @@
-// %BANNER_BEGIN%
-// ---------------------------------------------------------------------
-// %COPYRIGHT_BEGIN%
-//
-// Copyright (c) 2017 Magic Leap, Inc. (COMPANY) All Rights Reserved.
-// Magic Leap, Inc. Confidential and Proprietary
-//
-// NOTICE: All information contained herein is, and remains the property
-// of COMPANY. The intellectual and technical concepts contained herein
-// are proprietary to COMPANY and may be covered by U.S. and Foreign
-// Patents, patents in process, and are protected by trade secret or
-// copyright law. Dissemination of this information or reproduction of
-// this material is strictly forbidden unless prior written permission is
-// obtained from COMPANY. Access to the source code contained herein is
-// hereby forbidden to anyone except current COMPANY employees, managers
-// or contractors who have executed Confidentiality and Non-disclosure
-// agreements explicitly covering such access.
-//
-// The copyright notice above does not evidence any actual or intended
-// publication or disclosure of this source code, which includes
-// information that is confidential and/or proprietary, and is a trade
-// secret, of COMPANY. ANY REPRODUCTION, MODIFICATION, DISTRIBUTION,
-// PUBLIC PERFORMANCE, OR PUBLIC DISPLAY OF OR THROUGH USE OF THIS
-// SOURCE CODE WITHOUT THE EXPRESS WRITTEN CONSENT OF COMPANY IS
-// STRICTLY PROHIBITED, AND IN VIOLATION OF APPLICABLE LAWS AND
-// INTERNATIONAL TREATIES. THE RECEIPT OR POSSESSION OF THIS SOURCE
-// CODE AND/OR RELATED INFORMATION DOES NOT CONVEY OR IMPLY ANY RIGHTS
-// TO REPRODUCE, DISCLOSE OR DISTRIBUTE ITS CONTENTS, OR TO MANUFACTURE,
-// USE, OR SELL ANYTHING THAT IT MAY DESCRIBE, IN WHOLE OR IN PART.
-//
-// %COPYRIGHT_END%
-// --------------------------------------------------------------------
-// %BANNER_END%
+// Copyright 1998-2018 Epic Games, Inc. All Rights Reserved.
 
 #include "MagicLeapControllerFunctionLibrary.h"
 #include "IMagicLeapControllerPlugin.h"
diff --git a/Engine/Plugins/Lumin/MagicLeap/Source/MagicLeapController/Private/MagicLeapInputState.h b/Engine/Plugins/Lumin/MagicLeap/Source/MagicLeapController/Private/MagicLeapInputState.h
--- a/Engine/Plugins/Lumin/MagicLeap/Source/MagicLeapController/Private/MagicLeapInputState.h
+++ b/Engine/Plugins/Lumin/MagicLeap/Source/MagicLeapController/Private/MagicLeapInputState.h
@@ -1,36 +1,4 @@
-// %BANNER_BEGIN%
-// ---------------------------------------------------------------------
-// %COPYRIGHT_BEGIN%
-//
-// Copyright (c) 2017 Magic Leap, Inc. (COMPANY) All Rights Reserved.
-// Magic Leap, Inc. Confidential and Proprietary
-//
-// NOTICE: All information contained herein is, and remains the property
-// of COMPANY. The intellectual and technical concepts contained herein
-// are proprietary to COMPANY and may be covered by U.S. and Foreign
-// Patents, patents in process, and are protected by trade secret or
-// copyright law. Dissemination of this information or reproduction of
-// this material is strictly forbidden unless prior written permission is
-// obtained from COMPANY. Access to the source code contained herein is
-// hereby forbidden to anyone except current COMPANY employees, managers
-// or contractors who have executed Confidentiality and Non-disclosure
-// agreements explicitly covering such access.
-//
-// The copyright notice above does not evidence any actual or intended
-// publication or disclosure of this source code, which includes
-// information that is confidential and/or proprietary, and is a trade
-// secret, of COMPANY. ANY REPRODUCTION, MODIFICATION, DISTRIBUTION,
-// PUBLIC PERFORMANCE, OR PUBLIC DISPLAY OF OR THROUGH USE OF THIS
-// SOURCE CODE WITHOUT THE EXPRESS WRITTEN CONSENT OF COMPANY IS
-// STRICTLY PROHIBITED, AND IN VIOLATION OF APPLICABLE LAWS AND
-// INTERNATIONAL TREATIES. THE RECEIPT OR POSSESSION OF THIS SOURCE
-// CODE AND/OR RELATED INFORMATION DOES NOT CONVEY OR IMPLY ANY RIGHTS
-// TO REPRODUCE, DISCLOSE OR DISTRIBUTE ITS CONTENTS, OR TO MANUFACTURE,
-// USE, OR SELL ANYTHING THAT IT MAY DESCRIBE, IN WHOLE OR IN PART.
-//
-// %COPYRIGHT_END%
-// --------------------------------------------------------------------
-// %BANNER_END%
+// Copyright 1998-2018 Epic Games, Inc. All Rights Reserved.
 
 #pragma once
 
diff --git a/Engine/Plugins/Lumin/MagicLeap/Source/MagicLeapController/Private/Tests/MagicLeapControllerTests.cpp b/Engine/Plugins/Lumin/MagicLeap/Source/MagicLeapController/Private/Tests/MagicLeapControllerTests.cpp
--- a/Engine/Plugins/Lumin/MagicLeap/Source/MagicLeapController/Private/Tests/MagicLeapControllerTests.cpp
+++ b/Engine/Plugins/Lumin/MagicLeap/Source/MagicLeapController/Private/Tests/MagicLeapControllerTests.cpp
@@ -1,36 +1,4 @@
-// %BANNER_BEGIN%
-// ---------------------------------------------------------------------
-// %COPYRIGHT_BEGIN%
-//
-// Copyright (c) 2017 Magic Leap, Inc. (COMPANY) All Rights Reserved.
-// Magic Leap, Inc. Confidential and Proprietary
-//
-// NOTICE: All information contained herein is, and remains the property
-// of COMPANY. The intellectual and technical concepts contained herein
-// are proprietary to COMPANY and may be covered by U.S. and Foreign
-// Patents, patents in process, and are protected by trade secret or
-// copyright law. Dissemination of this information or reproduction of
-// this material is strictly forbidden unless prior written permission is
-// obtained from COMPANY. Access to the source code contained herein is
-// hereby forbidden to anyone except current COMPANY employees, managers
-// or contractors who have executed Confidentiality and Non-disclosure
-// agreements explicitly covering such access.
-//
-// The copyright notice above does not evidence any actual or intended
-// publication or disclosure of this source code, which includes
-// information that is confidential and/or proprietary, and is a trade
-// secret, of COMPANY. ANY REPRODUCTION, MODIFICATION, DISTRIBUTION,
-// PUBLIC PERFORMANCE, OR PUBLIC DISPLAY OF OR THROUGH USE OF THIS
-// SOURCE CODE WITHOUT THE EXPRESS WRITTEN CONSENT OF COMPANY IS
-// STRICTLY PROHIBITED, AND IN VIOLATION OF APPLICABLE LAWS AND
-// INTERNATIONAL TREATIES. THE RECEIPT OR POSSESSION OF THIS SOURCE
-// CODE AND/OR RELATED INFORMATION DOES NOT CONVEY OR IMPLY ANY RIGHTS
-// TO REPRODUCE, DISCLOSE OR DISTRIBUTE ITS CONTENTS, OR TO MANUFACTURE,
-// USE, OR SELL ANYTHING THAT IT MAY DESCRIBE, IN WHOLE OR IN PART.
-//
-// %COPYRIGHT_END%
-// --------------------------------------------------------------------
-// %BANNER_END%
+// Copyright 1998-2018 Epic Games, Inc. All Rights Reserved.
 
 #include "UObject/Package.h"
 #include "Misc/AutomationTest.h"
diff --git a/Engine/Plugins/Lumin/MagicLeap/Source/MagicLeapController/Private/TouchpadGesturesComponent.cpp b/Engine/Plugins/Lumin/MagicLeap/Source/MagicLeapController/Private/TouchpadGesturesComponent.cpp
--- a/Engine/Plugins/Lumin/MagicLeap/Source/MagicLeapController/Private/TouchpadGesturesComponent.cpp
+++ b/Engine/Plugins/Lumin/MagicLeap/Source/MagicLeapController/Private/TouchpadGesturesComponent.cpp
@@ -1,36 +1,4 @@
-// %BANNER_BEGIN%
-// ---------------------------------------------------------------------
-// %COPYRIGHT_BEGIN%
-//
-// Copyright (c) 2017 Magic Leap, Inc. (COMPANY) All Rights Reserved.
-// Magic Leap, Inc. Confidential and Proprietary
-//
-// NOTICE: All information contained herein is, and remains the property
-// of COMPANY. The intellectual and technical concepts contained herein
-// are proprietary to COMPANY and may be covered by U.S. and Foreign
-// Patents, patents in process, and are protected by trade secret or
-// copyright law. Dissemination of this information or reproduction of
-// this material is strictly forbidden unless prior written permission is
-// obtained from COMPANY. Access to the source code contained herein is
-// hereby forbidden to anyone except current COMPANY employees, managers
-// or contractors who have executed Confidentiality and Non-disclosure
-// agreements explicitly covering such access.
-//
-// The copyright notice above does not evidence any actual or intended
-// publication or disclosure of this source code, which includes
-// information that is confidential and/or proprietary, and is a trade
-// secret, of COMPANY. ANY REPRODUCTION, MODIFICATION, DISTRIBUTION,
-// PUBLIC PERFORMANCE, OR PUBLIC DISPLAY OF OR THROUGH USE OF THIS
-// SOURCE CODE WITHOUT THE EXPRESS WRITTEN CONSENT OF COMPANY IS
-// STRICTLY PROHIBITED, AND IN VIOLATION OF APPLICABLE LAWS AND
-// INTERNATIONAL TREATIES. THE RECEIPT OR POSSESSION OF THIS SOURCE
-// CODE AND/OR RELATED INFORMATION DOES NOT CONVEY OR IMPLY ANY RIGHTS
-// TO REPRODUCE, DISCLOSE OR DISTRIBUTE ITS CONTENTS, OR TO MANUFACTURE,
-// USE, OR SELL ANYTHING THAT IT MAY DESCRIBE, IN WHOLE OR IN PART.
-//
-// %COPYRIGHT_END%
-// --------------------------------------------------------------------
-// %BANNER_END%
+// Copyright 1998-2018 Epic Games, Inc. All Rights Reserved.
 
 #include "TouchpadGesturesComponent.h"
 #include "IMagicLeapControllerPlugin.h"
diff --git a/Engine/Plugins/Lumin/MagicLeap/Source/MagicLeapController/Public/IMagicLeapControllerPlugin.h b/Engine/Plugins/Lumin/MagicLeap/Source/MagicLeapController/Public/IMagicLeapControllerPlugin.h
--- a/Engine/Plugins/Lumin/MagicLeap/Source/MagicLeapController/Public/IMagicLeapControllerPlugin.h
+++ b/Engine/Plugins/Lumin/MagicLeap/Source/MagicLeapController/Public/IMagicLeapControllerPlugin.h
@@ -1,36 +1,4 @@
-// %BANNER_BEGIN%
-// ---------------------------------------------------------------------
-// %COPYRIGHT_BEGIN%
-//
-// Copyright (c) 2017 Magic Leap, Inc. (COMPANY) All Rights Reserved.
-// Magic Leap, Inc. Confidential and Proprietary
-//
-// NOTICE: All information contained herein is, and remains the property
-// of COMPANY. The intellectual and technical concepts contained herein
-// are proprietary to COMPANY and may be covered by U.S. and Foreign
-// Patents, patents in process, and are protected by trade secret or
-// copyright law. Dissemination of this information or reproduction of
-// this material is strictly forbidden unless prior written permission is
-// obtained from COMPANY. Access to the source code contained herein is
-// hereby forbidden to anyone except current COMPANY employees, managers
-// or contractors who have executed Confidentiality and Non-disclosure
-// agreements explicitly covering such access.
-//
-// The copyright notice above does not evidence any actual or intended
-// publication or disclosure of this source code, which includes
-// information that is confidential and/or proprietary, and is a trade
-// secret, of COMPANY. ANY REPRODUCTION, MODIFICATION, DISTRIBUTION,
-// PUBLIC PERFORMANCE, OR PUBLIC DISPLAY OF OR THROUGH USE OF THIS
-// SOURCE CODE WITHOUT THE EXPRESS WRITTEN CONSENT OF COMPANY IS
-// STRICTLY PROHIBITED, AND IN VIOLATION OF APPLICABLE LAWS AND
-// INTERNATIONAL TREATIES. THE RECEIPT OR POSSESSION OF THIS SOURCE
-// CODE AND/OR RELATED INFORMATION DOES NOT CONVEY OR IMPLY ANY RIGHTS
-// TO REPRODUCE, DISCLOSE OR DISTRIBUTE ITS CONTENTS, OR TO MANUFACTURE,
-// USE, OR SELL ANYTHING THAT IT MAY DESCRIBE, IN WHOLE OR IN PART.
-//
-// %COPYRIGHT_END%
-// --------------------------------------------------------------------
-// %BANNER_END%
+// Copyright 1998-2018 Epic Games, Inc. All Rights Reserved.
 
 #pragma once
 
diff --git a/Engine/Plugins/Lumin/MagicLeap/Source/MagicLeapController/Public/MagicLeapControllerFunctionLibrary.h b/Engine/Plugins/Lumin/MagicLeap/Source/MagicLeapController/Public/MagicLeapControllerFunctionLibrary.h
--- a/Engine/Plugins/Lumin/MagicLeap/Source/MagicLeapController/Public/MagicLeapControllerFunctionLibrary.h
+++ b/Engine/Plugins/Lumin/MagicLeap/Source/MagicLeapController/Public/MagicLeapControllerFunctionLibrary.h
@@ -1,36 +1,4 @@
-// %BANNER_BEGIN%
-// ---------------------------------------------------------------------
-// %COPYRIGHT_BEGIN%
-//
-// Copyright (c) 2017 Magic Leap, Inc. (COMPANY) All Rights Reserved.
-// Magic Leap, Inc. Confidential and Proprietary
-//
-// NOTICE: All information contained herein is, and remains the property
-// of COMPANY. The intellectual and technical concepts contained herein
-// are proprietary to COMPANY and may be covered by U.S. and Foreign
-// Patents, patents in process, and are protected by trade secret or
-// copyright law. Dissemination of this information or reproduction of
-// this material is strictly forbidden unless prior written permission is
-// obtained from COMPANY. Access to the source code contained herein is
-// hereby forbidden to anyone except current COMPANY employees, managers
-// or contractors who have executed Confidentiality and Non-disclosure
-// agreements explicitly covering such access.
-//
-// The copyright notice above does not evidence any actual or intended
-// publication or disclosure of this source code, which includes
-// information that is confidential and/or proprietary, and is a trade
-// secret, of COMPANY. ANY REPRODUCTION, MODIFICATION, DISTRIBUTION,
-// PUBLIC PERFORMANCE, OR PUBLIC DISPLAY OF OR THROUGH USE OF THIS
-// SOURCE CODE WITHOUT THE EXPRESS WRITTEN CONSENT OF COMPANY IS
-// STRICTLY PROHIBITED, AND IN VIOLATION OF APPLICABLE LAWS AND
-// INTERNATIONAL TREATIES. THE RECEIPT OR POSSESSION OF THIS SOURCE
-// CODE AND/OR RELATED INFORMATION DOES NOT CONVEY OR IMPLY ANY RIGHTS
-// TO REPRODUCE, DISCLOSE OR DISTRIBUTE ITS CONTENTS, OR TO MANUFACTURE,
-// USE, OR SELL ANYTHING THAT IT MAY DESCRIBE, IN WHOLE OR IN PART.
-//
-// %COPYRIGHT_END%
-// --------------------------------------------------------------------
-// %BANNER_END%
+// Copyright 1998-2018 Epic Games, Inc. All Rights Reserved.
 
 #pragma once
 
diff --git a/Engine/Plugins/Lumin/MagicLeap/Source/MagicLeapController/Public/MagicLeapControllerKeys.h b/Engine/Plugins/Lumin/MagicLeap/Source/MagicLeapController/Public/MagicLeapControllerKeys.h
--- a/Engine/Plugins/Lumin/MagicLeap/Source/MagicLeapController/Public/MagicLeapControllerKeys.h
+++ b/Engine/Plugins/Lumin/MagicLeap/Source/MagicLeapController/Public/MagicLeapControllerKeys.h
@@ -1,36 +1,4 @@
-// %BANNER_BEGIN%
-// ---------------------------------------------------------------------
-// %COPYRIGHT_BEGIN%
-//
-// Copyright (c) 2017 Magic Leap, Inc. (COMPANY) All Rights Reserved.
-// Magic Leap, Inc. Confidential and Proprietary
-//
-// NOTICE: All information contained herein is, and remains the property
-// of COMPANY. The intellectual and technical concepts contained herein
-// are proprietary to COMPANY and may be covered by U.S. and Foreign
-// Patents, patents in process, and are protected by trade secret or
-// copyright law. Dissemination of this information or reproduction of
-// this material is strictly forbidden unless prior written permission is
-// obtained from COMPANY. Access to the source code contained herein is
-// hereby forbidden to anyone except current COMPANY employees, managers
-// or contractors who have executed Confidentiality and Non-disclosure
-// agreements explicitly covering such access.
-//
-// The copyright notice above does not evidence any actual or intended
-// publication or disclosure of this source code, which includes
-// information that is confidential and/or proprietary, and is a trade
-// secret, of COMPANY. ANY REPRODUCTION, MODIFICATION, DISTRIBUTION,
-// PUBLIC PERFORMANCE, OR PUBLIC DISPLAY OF OR THROUGH USE OF THIS
-// SOURCE CODE WITHOUT THE EXPRESS WRITTEN CONSENT OF COMPANY IS
-// STRICTLY PROHIBITED, AND IN VIOLATION OF APPLICABLE LAWS AND
-// INTERNATIONAL TREATIES. THE RECEIPT OR POSSESSION OF THIS SOURCE
-// CODE AND/OR RELATED INFORMATION DOES NOT CONVEY OR IMPLY ANY RIGHTS
-// TO REPRODUCE, DISCLOSE OR DISTRIBUTE ITS CONTENTS, OR TO MANUFACTURE,
-// USE, OR SELL ANYTHING THAT IT MAY DESCRIBE, IN WHOLE OR IN PART.
-//
-// %COPYRIGHT_END%
-// --------------------------------------------------------------------
-// %BANNER_END%
+// Copyright 1998-2018 Epic Games, Inc. All Rights Reserved.
 
 #pragma once
 
diff --git a/Engine/Plugins/Lumin/MagicLeap/Source/MagicLeapController/Public/TouchpadGesturesComponent.h b/Engine/Plugins/Lumin/MagicLeap/Source/MagicLeapController/Public/TouchpadGesturesComponent.h
--- a/Engine/Plugins/Lumin/MagicLeap/Source/MagicLeapController/Public/TouchpadGesturesComponent.h
+++ b/Engine/Plugins/Lumin/MagicLeap/Source/MagicLeapController/Public/TouchpadGesturesComponent.h
@@ -1,36 +1,4 @@
-// %BANNER_BEGIN%
-// ---------------------------------------------------------------------
-// %COPYRIGHT_BEGIN%
-//
-// Copyright (c) 2017 Magic Leap, Inc. (COMPANY) All Rights Reserved.
-// Magic Leap, Inc. Confidential and Proprietary
-//
-// NOTICE: All information contained herein is, and remains the property
-// of COMPANY. The intellectual and technical concepts contained herein
-// are proprietary to COMPANY and may be covered by U.S. and Foreign
-// Patents, patents in process, and are protected by trade secret or
-// copyright law. Dissemination of this information or reproduction of
-// this material is strictly forbidden unless prior written permission is
-// obtained from COMPANY. Access to the source code contained herein is
-// hereby forbidden to anyone except current COMPANY employees, managers
-// or contractors who have executed Confidentiality and Non-disclosure
-// agreements explicitly covering such access.
-//
-// The copyright notice above does not evidence any actual or intended
-// publication or disclosure of this source code, which includes
-// information that is confidential and/or proprietary, and is a trade
-// secret, of COMPANY. ANY REPRODUCTION, MODIFICATION, DISTRIBUTION,
-// PUBLIC PERFORMANCE, OR PUBLIC DISPLAY OF OR THROUGH USE OF THIS
-// SOURCE CODE WITHOUT THE EXPRESS WRITTEN CONSENT OF COMPANY IS
-// STRICTLY PROHIBITED, AND IN VIOLATION OF APPLICABLE LAWS AND
-// INTERNATIONAL TREATIES. THE RECEIPT OR POSSESSION OF THIS SOURCE
-// CODE AND/OR RELATED INFORMATION DOES NOT CONVEY OR IMPLY ANY RIGHTS
-// TO REPRODUCE, DISCLOSE OR DISTRIBUTE ITS CONTENTS, OR TO MANUFACTURE,
-// USE, OR SELL ANYTHING THAT IT MAY DESCRIBE, IN WHOLE OR IN PART.
-//
-// %COPYRIGHT_END%
-// --------------------------------------------------------------------
-// %BANNER_END%
+// Copyright 1998-2018 Epic Games, Inc. All Rights Reserved.
 
 #pragma once
 
diff --git a/Engine/Plugins/Lumin/MagicLeap/Source/MagicLeapEyeTracker/MagicLeapEyeTracker.Build.cs b/Engine/Plugins/Lumin/MagicLeap/Source/MagicLeapEyeTracker/MagicLeapEyeTracker.Build.cs
--- a/Engine/Plugins/Lumin/MagicLeap/Source/MagicLeapEyeTracker/MagicLeapEyeTracker.Build.cs
+++ b/Engine/Plugins/Lumin/MagicLeap/Source/MagicLeapEyeTracker/MagicLeapEyeTracker.Build.cs
@@ -1,36 +1,4 @@
-// %BANNER_BEGIN%
-// ---------------------------------------------------------------------
-// %COPYRIGHT_BEGIN%
-//
-// Copyright (c) 2017 Magic Leap, Inc. (COMPANY) All Rights Reserved.
-// Magic Leap, Inc. Confidential and Proprietary
-//
-// NOTICE: All information contained herein is, and remains the property
-// of COMPANY. The intellectual and technical concepts contained herein
-// are proprietary to COMPANY and may be covered by U.S. and Foreign
-// Patents, patents in process, and are protected by trade secret or
-// copyright law. Dissemination of this information or reproduction of
-// this material is strictly forbidden unless prior written permission is
-// obtained from COMPANY. Access to the source code contained herein is
-// hereby forbidden to anyone except current COMPANY employees, managers
-// or contractors who have executed Confidentiality and Non-disclosure
-// agreements explicitly covering such access.
-//
-// The copyright notice above does not evidence any actual or intended
-// publication or disclosure of this source code, which includes
-// information that is confidential and/or proprietary, and is a trade
-// secret, of COMPANY. ANY REPRODUCTION, MODIFICATION, DISTRIBUTION,
-// PUBLIC PERFORMANCE, OR PUBLIC DISPLAY OF OR THROUGH USE OF THIS
-// SOURCE CODE WITHOUT THE EXPRESS WRITTEN CONSENT OF COMPANY IS
-// STRICTLY PROHIBITED, AND IN VIOLATION OF APPLICABLE LAWS AND
-// INTERNATIONAL TREATIES. THE RECEIPT OR POSSESSION OF THIS SOURCE
-// CODE AND/OR RELATED INFORMATION DOES NOT CONVEY OR IMPLY ANY RIGHTS
-// TO REPRODUCE, DISCLOSE OR DISTRIBUTE ITS CONTENTS, OR TO MANUFACTURE,
-// USE, OR SELL ANYTHING THAT IT MAY DESCRIBE, IN WHOLE OR IN PART.
-//
-// %COPYRIGHT_END%
-// --------------------------------------------------------------------
-// %BANNER_END%
+// Copyright 1998-2018 Epic Games, Inc. All Rights Reserved.
 
 using System.IO;
 
diff --git a/Engine/Plugins/Lumin/MagicLeap/Source/MagicLeapEyeTracker/Private/MagicLeapEyeTrackerModule.cpp b/Engine/Plugins/Lumin/MagicLeap/Source/MagicLeapEyeTracker/Private/MagicLeapEyeTrackerModule.cpp
--- a/Engine/Plugins/Lumin/MagicLeap/Source/MagicLeapEyeTracker/Private/MagicLeapEyeTrackerModule.cpp
+++ b/Engine/Plugins/Lumin/MagicLeap/Source/MagicLeapEyeTracker/Private/MagicLeapEyeTrackerModule.cpp
@@ -1,36 +1,4 @@
-// %BANNER_BEGIN%
-// ---------------------------------------------------------------------
-// %COPYRIGHT_BEGIN%
-//
-// Copyright (c) 2017 Magic Leap, Inc. (COMPANY) All Rights Reserved.
-// Magic Leap, Inc. Confidential and Proprietary
-//
-// NOTICE: All information contained herein is, and remains the property
-// of COMPANY. The intellectual and technical concepts contained herein
-// are proprietary to COMPANY and may be covered by U.S. and Foreign
-// Patents, patents in process, and are protected by trade secret or
-// copyright law. Dissemination of this information or reproduction of
-// this material is strictly forbidden unless prior written permission is
-// obtained from COMPANY. Access to the source code contained herein is
-// hereby forbidden to anyone except current COMPANY employees, managers
-// or contractors who have executed Confidentiality and Non-disclosure
-// agreements explicitly covering such access.
-//
-// The copyright notice above does not evidence any actual or intended
-// publication or disclosure of this source code, which includes
-// information that is confidential and/or proprietary, and is a trade
-// secret, of COMPANY. ANY REPRODUCTION, MODIFICATION, DISTRIBUTION,
-// PUBLIC PERFORMANCE, OR PUBLIC DISPLAY OF OR THROUGH USE OF THIS
-// SOURCE CODE WITHOUT THE EXPRESS WRITTEN CONSENT OF COMPANY IS
-// STRICTLY PROHIBITED, AND IN VIOLATION OF APPLICABLE LAWS AND
-// INTERNATIONAL TREATIES. THE RECEIPT OR POSSESSION OF THIS SOURCE
-// CODE AND/OR RELATED INFORMATION DOES NOT CONVEY OR IMPLY ANY RIGHTS
-// TO REPRODUCE, DISCLOSE OR DISTRIBUTE ITS CONTENTS, OR TO MANUFACTURE,
-// USE, OR SELL ANYTHING THAT IT MAY DESCRIBE, IN WHOLE OR IN PART.
-//
-// %COPYRIGHT_END%
-// --------------------------------------------------------------------
-// %BANNER_END%
+// Copyright 1998-2018 Epic Games, Inc. All Rights Reserved.
 
 #include "MagicLeapEyeTrackerModule.h"
 #include "GameFramework/HUD.h"
diff --git a/Engine/Plugins/Lumin/MagicLeap/Source/MagicLeapEyeTracker/Private/MagicLeapVREyeTracker.cpp b/Engine/Plugins/Lumin/MagicLeap/Source/MagicLeapEyeTracker/Private/MagicLeapVREyeTracker.cpp
--- a/Engine/Plugins/Lumin/MagicLeap/Source/MagicLeapEyeTracker/Private/MagicLeapVREyeTracker.cpp
+++ b/Engine/Plugins/Lumin/MagicLeap/Source/MagicLeapEyeTracker/Private/MagicLeapVREyeTracker.cpp
@@ -1,36 +1,4 @@
-// %BANNER_BEGIN%
-// ---------------------------------------------------------------------
-// %COPYRIGHT_BEGIN%
-//
-// Copyright (c) 2017 Magic Leap, Inc. (COMPANY) All Rights Reserved.
-// Magic Leap, Inc. Confidential and Proprietary
-//
-// NOTICE: All information contained herein is, and remains the property
-// of COMPANY. The intellectual and technical concepts contained herein
-// are proprietary to COMPANY and may be covered by U.S. and Foreign
-// Patents, patents in process, and are protected by trade secret or
-// copyright law. Dissemination of this information or reproduction of
-// this material is strictly forbidden unless prior written permission is
-// obtained from COMPANY. Access to the source code contained herein is
-// hereby forbidden to anyone except current COMPANY employees, managers
-// or contractors who have executed Confidentiality and Non-disclosure
-// agreements explicitly covering such access.
-//
-// The copyright notice above does not evidence any actual or intended
-// publication or disclosure of this source code, which includes
-// information that is confidential and/or proprietary, and is a trade
-// secret, of COMPANY. ANY REPRODUCTION, MODIFICATION, DISTRIBUTION,
-// PUBLIC PERFORMANCE, OR PUBLIC DISPLAY OF OR THROUGH USE OF THIS
-// SOURCE CODE WITHOUT THE EXPRESS WRITTEN CONSENT OF COMPANY IS
-// STRICTLY PROHIBITED, AND IN VIOLATION OF APPLICABLE LAWS AND
-// INTERNATIONAL TREATIES. THE RECEIPT OR POSSESSION OF THIS SOURCE
-// CODE AND/OR RELATED INFORMATION DOES NOT CONVEY OR IMPLY ANY RIGHTS
-// TO REPRODUCE, DISCLOSE OR DISTRIBUTE ITS CONTENTS, OR TO MANUFACTURE,
-// USE, OR SELL ANYTHING THAT IT MAY DESCRIBE, IN WHOLE OR IN PART.
-//
-// %COPYRIGHT_END%
-// --------------------------------------------------------------------
-// %BANNER_END%
+// Copyright 1998-2018 Epic Games, Inc. All Rights Reserved.
 
 #include "MagicLeapVREyeTracker.h"
 #include "IEyeTrackerModule.h"
diff --git a/Engine/Plugins/Lumin/MagicLeap/Source/MagicLeapEyeTracker/Private/MagicLeapVREyeTracker.h b/Engine/Plugins/Lumin/MagicLeap/Source/MagicLeapEyeTracker/Private/MagicLeapVREyeTracker.h
--- a/Engine/Plugins/Lumin/MagicLeap/Source/MagicLeapEyeTracker/Private/MagicLeapVREyeTracker.h
+++ b/Engine/Plugins/Lumin/MagicLeap/Source/MagicLeapEyeTracker/Private/MagicLeapVREyeTracker.h
@@ -1,36 +1,4 @@
-// %BANNER_BEGIN%
-// ---------------------------------------------------------------------
-// %COPYRIGHT_BEGIN%
-//
-// Copyright (c) 2017 Magic Leap, Inc. (COMPANY) All Rights Reserved.
-// Magic Leap, Inc. Confidential and Proprietary
-//
-// NOTICE: All information contained herein is, and remains the property
-// of COMPANY. The intellectual and technical concepts contained herein
-// are proprietary to COMPANY and may be covered by U.S. and Foreign
-// Patents, patents in process, and are protected by trade secret or
-// copyright law. Dissemination of this information or reproduction of
-// this material is strictly forbidden unless prior written permission is
-// obtained from COMPANY. Access to the source code contained herein is
-// hereby forbidden to anyone except current COMPANY employees, managers
-// or contractors who have executed Confidentiality and Non-disclosure
-// agreements explicitly covering such access.
-//
-// The copyright notice above does not evidence any actual or intended
-// publication or disclosure of this source code, which includes
-// information that is confidential and/or proprietary, and is a trade
-// secret, of COMPANY. ANY REPRODUCTION, MODIFICATION, DISTRIBUTION,
-// PUBLIC PERFORMANCE, OR PUBLIC DISPLAY OF OR THROUGH USE OF THIS
-// SOURCE CODE WITHOUT THE EXPRESS WRITTEN CONSENT OF COMPANY IS
-// STRICTLY PROHIBITED, AND IN VIOLATION OF APPLICABLE LAWS AND
-// INTERNATIONAL TREATIES. THE RECEIPT OR POSSESSION OF THIS SOURCE
-// CODE AND/OR RELATED INFORMATION DOES NOT CONVEY OR IMPLY ANY RIGHTS
-// TO REPRODUCE, DISCLOSE OR DISTRIBUTE ITS CONTENTS, OR TO MANUFACTURE,
-// USE, OR SELL ANYTHING THAT IT MAY DESCRIBE, IN WHOLE OR IN PART.
-//
-// %COPYRIGHT_END%
-// --------------------------------------------------------------------
-// %BANNER_END%
+// Copyright 1998-2018 Epic Games, Inc. All Rights Reserved.
 
 #pragma once
 
diff --git a/Engine/Plugins/Lumin/MagicLeap/Source/MagicLeapEyeTracker/Public/IMagicLeapVREyeTracker.h b/Engine/Plugins/Lumin/MagicLeap/Source/MagicLeapEyeTracker/Public/IMagicLeapVREyeTracker.h
--- a/Engine/Plugins/Lumin/MagicLeap/Source/MagicLeapEyeTracker/Public/IMagicLeapVREyeTracker.h
+++ b/Engine/Plugins/Lumin/MagicLeap/Source/MagicLeapEyeTracker/Public/IMagicLeapVREyeTracker.h
@@ -1,36 +1,4 @@
-// %BANNER_BEGIN%
-// ---------------------------------------------------------------------
-// %COPYRIGHT_BEGIN%
-//
-// Copyright (c) 2017 Magic Leap, Inc. (COMPANY) All Rights Reserved.
-// Magic Leap, Inc. Confidential and Proprietary
-//
-// NOTICE: All information contained herein is, and remains the property
-// of COMPANY. The intellectual and technical concepts contained herein
-// are proprietary to COMPANY and may be covered by U.S. and Foreign
-// Patents, patents in process, and are protected by trade secret or
-// copyright law. Dissemination of this information or reproduction of
-// this material is strictly forbidden unless prior written permission is
-// obtained from COMPANY. Access to the source code contained herein is
-// hereby forbidden to anyone except current COMPANY employees, managers
-// or contractors who have executed Confidentiality and Non-disclosure
-// agreements explicitly covering such access.
-//
-// The copyright notice above does not evidence any actual or intended
-// publication or disclosure of this source code, which includes
-// information that is confidential and/or proprietary, and is a trade
-// secret, of COMPANY. ANY REPRODUCTION, MODIFICATION, DISTRIBUTION,
-// PUBLIC PERFORMANCE, OR PUBLIC DISPLAY OF OR THROUGH USE OF THIS
-// SOURCE CODE WITHOUT THE EXPRESS WRITTEN CONSENT OF COMPANY IS
-// STRICTLY PROHIBITED, AND IN VIOLATION OF APPLICABLE LAWS AND
-// INTERNATIONAL TREATIES. THE RECEIPT OR POSSESSION OF THIS SOURCE
-// CODE AND/OR RELATED INFORMATION DOES NOT CONVEY OR IMPLY ANY RIGHTS
-// TO REPRODUCE, DISCLOSE OR DISTRIBUTE ITS CONTENTS, OR TO MANUFACTURE,
-// USE, OR SELL ANYTHING THAT IT MAY DESCRIBE, IN WHOLE OR IN PART.
-//
-// %COPYRIGHT_END%
-// --------------------------------------------------------------------
-// %BANNER_END%
+// Copyright 1998-2018 Epic Games, Inc. All Rights Reserved.
 
 #pragma once
 
diff --git a/Engine/Plugins/Lumin/MagicLeap/Source/MagicLeapEyeTracker/Public/MagicLeapEyeTrackerModule.h b/Engine/Plugins/Lumin/MagicLeap/Source/MagicLeapEyeTracker/Public/MagicLeapEyeTrackerModule.h
--- a/Engine/Plugins/Lumin/MagicLeap/Source/MagicLeapEyeTracker/Public/MagicLeapEyeTrackerModule.h
+++ b/Engine/Plugins/Lumin/MagicLeap/Source/MagicLeapEyeTracker/Public/MagicLeapEyeTrackerModule.h
@@ -1,36 +1,4 @@
-// %BANNER_BEGIN%
-// ---------------------------------------------------------------------
-// %COPYRIGHT_BEGIN%
-//
-// Copyright (c) 2017 Magic Leap, Inc. (COMPANY) All Rights Reserved.
-// Magic Leap, Inc. Confidential and Proprietary
-//
-// NOTICE: All information contained herein is, and remains the property
-// of COMPANY. The intellectual and technical concepts contained herein
-// are proprietary to COMPANY and may be covered by U.S. and Foreign
-// Patents, patents in process, and are protected by trade secret or
-// copyright law. Dissemination of this information or reproduction of
-// this material is strictly forbidden unless prior written permission is
-// obtained from COMPANY. Access to the source code contained herein is
-// hereby forbidden to anyone except current COMPANY employees, managers
-// or contractors who have executed Confidentiality and Non-disclosure
-// agreements explicitly covering such access.
-//
-// The copyright notice above does not evidence any actual or intended
-// publication or disclosure of this source code, which includes
-// information that is confidential and/or proprietary, and is a trade
-// secret, of COMPANY. ANY REPRODUCTION, MODIFICATION, DISTRIBUTION,
-// PUBLIC PERFORMANCE, OR PUBLIC DISPLAY OF OR THROUGH USE OF THIS
-// SOURCE CODE WITHOUT THE EXPRESS WRITTEN CONSENT OF COMPANY IS
-// STRICTLY PROHIBITED, AND IN VIOLATION OF APPLICABLE LAWS AND
-// INTERNATIONAL TREATIES. THE RECEIPT OR POSSESSION OF THIS SOURCE
-// CODE AND/OR RELATED INFORMATION DOES NOT CONVEY OR IMPLY ANY RIGHTS
-// TO REPRODUCE, DISCLOSE OR DISTRIBUTE ITS CONTENTS, OR TO MANUFACTURE,
-// USE, OR SELL ANYTHING THAT IT MAY DESCRIBE, IN WHOLE OR IN PART.
-//
-// %COPYRIGHT_END%
-// --------------------------------------------------------------------
-// %BANNER_END%
+// Copyright 1998-2018 Epic Games, Inc. All Rights Reserved.
 
 #pragma once
 
diff --git a/Engine/Plugins/Lumin/MagicLeap/Source/MagicLeapEyeTracker/Public/MagicLeapEyeTrackerTypes.h b/Engine/Plugins/Lumin/MagicLeap/Source/MagicLeapEyeTracker/Public/MagicLeapEyeTrackerTypes.h
--- a/Engine/Plugins/Lumin/MagicLeap/Source/MagicLeapEyeTracker/Public/MagicLeapEyeTrackerTypes.h
+++ b/Engine/Plugins/Lumin/MagicLeap/Source/MagicLeapEyeTracker/Public/MagicLeapEyeTrackerTypes.h
@@ -1,36 +1,4 @@
-// %BANNER_BEGIN%
-// ---------------------------------------------------------------------
-// %COPYRIGHT_BEGIN%
-//
-// Copyright (c) 2017 Magic Leap, Inc. (COMPANY) All Rights Reserved.
-// Magic Leap, Inc. Confidential and Proprietary
-//
-// NOTICE: All information contained herein is, and remains the property
-// of COMPANY. The intellectual and technical concepts contained herein
-// are proprietary to COMPANY and may be covered by U.S. and Foreign
-// Patents, patents in process, and are protected by trade secret or
-// copyright law. Dissemination of this information or reproduction of
-// this material is strictly forbidden unless prior written permission is
-// obtained from COMPANY. Access to the source code contained herein is
-// hereby forbidden to anyone except current COMPANY employees, managers
-// or contractors who have executed Confidentiality and Non-disclosure
-// agreements explicitly covering such access.
-//
-// The copyright notice above does not evidence any actual or intended
-// publication or disclosure of this source code, which includes
-// information that is confidential and/or proprietary, and is a trade
-// secret, of COMPANY. ANY REPRODUCTION, MODIFICATION, DISTRIBUTION,
-// PUBLIC PERFORMANCE, OR PUBLIC DISPLAY OF OR THROUGH USE OF THIS
-// SOURCE CODE WITHOUT THE EXPRESS WRITTEN CONSENT OF COMPANY IS
-// STRICTLY PROHIBITED, AND IN VIOLATION OF APPLICABLE LAWS AND
-// INTERNATIONAL TREATIES. THE RECEIPT OR POSSESSION OF THIS SOURCE
-// CODE AND/OR RELATED INFORMATION DOES NOT CONVEY OR IMPLY ANY RIGHTS
-// TO REPRODUCE, DISCLOSE OR DISTRIBUTE ITS CONTENTS, OR TO MANUFACTURE,
-// USE, OR SELL ANYTHING THAT IT MAY DESCRIBE, IN WHOLE OR IN PART.
-//
-// %COPYRIGHT_END%
-// --------------------------------------------------------------------
-// %BANNER_END%
+// Copyright 1998-2018 Epic Games, Inc. All Rights Reserved.
 
 #pragma once
 
diff --git a/Engine/Plugins/Lumin/MagicLeap/Source/MagicLeapGestures/MagicLeapGestures.Build.cs b/Engine/Plugins/Lumin/MagicLeap/Source/MagicLeapGestures/MagicLeapGestures.Build.cs
--- a/Engine/Plugins/Lumin/MagicLeap/Source/MagicLeapGestures/MagicLeapGestures.Build.cs
+++ b/Engine/Plugins/Lumin/MagicLeap/Source/MagicLeapGestures/MagicLeapGestures.Build.cs
@@ -1,36 +1,4 @@
-// %BANNER_BEGIN%
-// ---------------------------------------------------------------------
-// %COPYRIGHT_BEGIN%
-//
-// Copyright (c) 2017 Magic Leap, Inc. (COMPANY) All Rights Reserved.
-// Magic Leap, Inc. Confidential and Proprietary
-//
-// NOTICE: All information contained herein is, and remains the property
-// of COMPANY. The intellectual and technical concepts contained herein
-// are proprietary to COMPANY and may be covered by U.S. and Foreign
-// Patents, patents in process, and are protected by trade secret or
-// copyright law. Dissemination of this information or reproduction of
-// this material is strictly forbidden unless prior written permission is
-// obtained from COMPANY. Access to the source code contained herein is
-// hereby forbidden to anyone except current COMPANY employees, managers
-// or contractors who have executed Confidentiality and Non-disclosure
-// agreements explicitly covering such access.
-//
-// The copyright notice above does not evidence any actual or intended
-// publication or disclosure of this source code, which includes
-// information that is confidential and/or proprietary, and is a trade
-// secret, of COMPANY. ANY REPRODUCTION, MODIFICATION, DISTRIBUTION,
-// PUBLIC PERFORMANCE, OR PUBLIC DISPLAY OF OR THROUGH USE OF THIS
-// SOURCE CODE WITHOUT THE EXPRESS WRITTEN CONSENT OF COMPANY IS
-// STRICTLY PROHIBITED, AND IN VIOLATION OF APPLICABLE LAWS AND
-// INTERNATIONAL TREATIES. THE RECEIPT OR POSSESSION OF THIS SOURCE
-// CODE AND/OR RELATED INFORMATION DOES NOT CONVEY OR IMPLY ANY RIGHTS
-// TO REPRODUCE, DISCLOSE OR DISTRIBUTE ITS CONTENTS, OR TO MANUFACTURE,
-// USE, OR SELL ANYTHING THAT IT MAY DESCRIBE, IN WHOLE OR IN PART.
-//
-// %COPYRIGHT_END%
-// --------------------------------------------------------------------
-// %BANNER_END%
+// Copyright 1998-2018 Epic Games, Inc. All Rights Reserved.
 
 using System.IO;
 
diff --git a/Engine/Plugins/Lumin/MagicLeap/Source/MagicLeapGestures/Private/MagicLeapGestures.cpp b/Engine/Plugins/Lumin/MagicLeap/Source/MagicLeapGestures/Private/MagicLeapGestures.cpp
--- a/Engine/Plugins/Lumin/MagicLeap/Source/MagicLeapGestures/Private/MagicLeapGestures.cpp
+++ b/Engine/Plugins/Lumin/MagicLeap/Source/MagicLeapGestures/Private/MagicLeapGestures.cpp
@@ -1,36 +1,4 @@
-// %BANNER_BEGIN%
-// ---------------------------------------------------------------------
-// %COPYRIGHT_BEGIN%
-//
-// Copyright (c) 2017 Magic Leap, Inc. (COMPANY) All Rights Reserved.
-// Magic Leap, Inc. Confidential and Proprietary
-//
-// NOTICE: All information contained herein is, and remains the property
-// of COMPANY. The intellectual and technical concepts contained herein
-// are proprietary to COMPANY and may be covered by U.S. and Foreign
-// Patents, patents in process, and are protected by trade secret or
-// copyright law. Dissemination of this information or reproduction of
-// this material is strictly forbidden unless prior written permission is
-// obtained from COMPANY. Access to the source code contained herein is
-// hereby forbidden to anyone except current COMPANY employees, managers
-// or contractors who have executed Confidentiality and Non-disclosure
-// agreements explicitly covering such access.
-//
-// The copyright notice above does not evidence any actual or intended
-// publication or disclosure of this source code, which includes
-// information that is confidential and/or proprietary, and is a trade
-// secret, of COMPANY. ANY REPRODUCTION, MODIFICATION, DISTRIBUTION,
-// PUBLIC PERFORMANCE, OR PUBLIC DISPLAY OF OR THROUGH USE OF THIS
-// SOURCE CODE WITHOUT THE EXPRESS WRITTEN CONSENT OF COMPANY IS
-// STRICTLY PROHIBITED, AND IN VIOLATION OF APPLICABLE LAWS AND
-// INTERNATIONAL TREATIES. THE RECEIPT OR POSSESSION OF THIS SOURCE
-// CODE AND/OR RELATED INFORMATION DOES NOT CONVEY OR IMPLY ANY RIGHTS
-// TO REPRODUCE, DISCLOSE OR DISTRIBUTE ITS CONTENTS, OR TO MANUFACTURE,
-// USE, OR SELL ANYTHING THAT IT MAY DESCRIBE, IN WHOLE OR IN PART.
-//
-// %COPYRIGHT_END%
-// --------------------------------------------------------------------
-// %BANNER_END%
+// Copyright 1998-2018 Epic Games, Inc. All Rights Reserved.
 
 #include "MagicLeapGestures.h"
 #include "MagicLeapHMD.h"
diff --git a/Engine/Plugins/Lumin/MagicLeap/Source/MagicLeapGestures/Private/MagicLeapGestures.h b/Engine/Plugins/Lumin/MagicLeap/Source/MagicLeapGestures/Private/MagicLeapGestures.h
--- a/Engine/Plugins/Lumin/MagicLeap/Source/MagicLeapGestures/Private/MagicLeapGestures.h
+++ b/Engine/Plugins/Lumin/MagicLeap/Source/MagicLeapGestures/Private/MagicLeapGestures.h
@@ -1,36 +1,4 @@
-// %BANNER_BEGIN%
-// ---------------------------------------------------------------------
-// %COPYRIGHT_BEGIN%
-//
-// Copyright (c) 2017 Magic Leap, Inc. (COMPANY) All Rights Reserved.
-// Magic Leap, Inc. Confidential and Proprietary
-//
-// NOTICE: All information contained herein is, and remains the property
-// of COMPANY. The intellectual and technical concepts contained herein
-// are proprietary to COMPANY and may be covered by U.S. and Foreign
-// Patents, patents in process, and are protected by trade secret or
-// copyright law. Dissemination of this information or reproduction of
-// this material is strictly forbidden unless prior written permission is
-// obtained from COMPANY. Access to the source code contained herein is
-// hereby forbidden to anyone except current COMPANY employees, managers
-// or contractors who have executed Confidentiality and Non-disclosure
-// agreements explicitly covering such access.
-//
-// The copyright notice above does not evidence any actual or intended
-// publication or disclosure of this source code, which includes
-// information that is confidential and/or proprietary, and is a trade
-// secret, of COMPANY. ANY REPRODUCTION, MODIFICATION, DISTRIBUTION,
-// PUBLIC PERFORMANCE, OR PUBLIC DISPLAY OF OR THROUGH USE OF THIS
-// SOURCE CODE WITHOUT THE EXPRESS WRITTEN CONSENT OF COMPANY IS
-// STRICTLY PROHIBITED, AND IN VIOLATION OF APPLICABLE LAWS AND
-// INTERNATIONAL TREATIES. THE RECEIPT OR POSSESSION OF THIS SOURCE
-// CODE AND/OR RELATED INFORMATION DOES NOT CONVEY OR IMPLY ANY RIGHTS
-// TO REPRODUCE, DISCLOSE OR DISTRIBUTE ITS CONTENTS, OR TO MANUFACTURE,
-// USE, OR SELL ANYTHING THAT IT MAY DESCRIBE, IN WHOLE OR IN PART.
-//
-// %COPYRIGHT_END%
-// --------------------------------------------------------------------
-// %BANNER_END%
+// Copyright 1998-2018 Epic Games, Inc. All Rights Reserved.
 
 #pragma once
 
diff --git a/Engine/Plugins/Lumin/MagicLeap/Source/MagicLeapGestures/Private/MagicLeapGesturesFunctionLibrary.cpp b/Engine/Plugins/Lumin/MagicLeap/Source/MagicLeapGestures/Private/MagicLeapGesturesFunctionLibrary.cpp
--- a/Engine/Plugins/Lumin/MagicLeap/Source/MagicLeapGestures/Private/MagicLeapGesturesFunctionLibrary.cpp
+++ b/Engine/Plugins/Lumin/MagicLeap/Source/MagicLeapGestures/Private/MagicLeapGesturesFunctionLibrary.cpp
@@ -1,36 +1,4 @@
-// %BANNER_BEGIN%
-// ---------------------------------------------------------------------
-// %COPYRIGHT_BEGIN%
-//
-// Copyright (c) 2017 Magic Leap, Inc. (COMPANY) All Rights Reserved.
-// Magic Leap, Inc. Confidential and Proprietary
-//
-// NOTICE: All information contained herein is, and remains the property
-// of COMPANY. The intellectual and technical concepts contained herein
-// are proprietary to COMPANY and may be covered by U.S. and Foreign
-// Patents, patents in process, and are protected by trade secret or
-// copyright law. Dissemination of this information or reproduction of
-// this material is strictly forbidden unless prior written permission is
-// obtained from COMPANY. Access to the source code contained herein is
-// hereby forbidden to anyone except current COMPANY employees, managers
-// or contractors who have executed Confidentiality and Non-disclosure
-// agreements explicitly covering such access.
-//
-// The copyright notice above does not evidence any actual or intended
-// publication or disclosure of this source code, which includes
-// information that is confidential and/or proprietary, and is a trade
-// secret, of COMPANY. ANY REPRODUCTION, MODIFICATION, DISTRIBUTION,
-// PUBLIC PERFORMANCE, OR PUBLIC DISPLAY OF OR THROUGH USE OF THIS
-// SOURCE CODE WITHOUT THE EXPRESS WRITTEN CONSENT OF COMPANY IS
-// STRICTLY PROHIBITED, AND IN VIOLATION OF APPLICABLE LAWS AND
-// INTERNATIONAL TREATIES. THE RECEIPT OR POSSESSION OF THIS SOURCE
-// CODE AND/OR RELATED INFORMATION DOES NOT CONVEY OR IMPLY ANY RIGHTS
-// TO REPRODUCE, DISCLOSE OR DISTRIBUTE ITS CONTENTS, OR TO MANUFACTURE,
-// USE, OR SELL ANYTHING THAT IT MAY DESCRIBE, IN WHOLE OR IN PART.
-//
-// %COPYRIGHT_END%
-// --------------------------------------------------------------------
-// %BANNER_END%
+// Copyright 1998-2018 Epic Games, Inc. All Rights Reserved.
 
 #include "MagicLeapGesturesFunctionLibrary.h"
 #include "IMagicLeapGesturesPlugin.h"
diff --git a/Engine/Plugins/Lumin/MagicLeap/Source/MagicLeapGestures/Public/IMagicLeapGesturesPlugin.h b/Engine/Plugins/Lumin/MagicLeap/Source/MagicLeapGestures/Public/IMagicLeapGesturesPlugin.h
--- a/Engine/Plugins/Lumin/MagicLeap/Source/MagicLeapGestures/Public/IMagicLeapGesturesPlugin.h
+++ b/Engine/Plugins/Lumin/MagicLeap/Source/MagicLeapGestures/Public/IMagicLeapGesturesPlugin.h
@@ -1,36 +1,4 @@
-// %BANNER_BEGIN%
-// ---------------------------------------------------------------------
-// %COPYRIGHT_BEGIN%
-//
-// Copyright (c) 2017 Magic Leap, Inc. (COMPANY) All Rights Reserved.
-// Magic Leap, Inc. Confidential and Proprietary
-//
-// NOTICE: All information contained herein is, and remains the property
-// of COMPANY. The intellectual and technical concepts contained herein
-// are proprietary to COMPANY and may be covered by U.S. and Foreign
-// Patents, patents in process, and are protected by trade secret or
-// copyright law. Dissemination of this information or reproduction of
-// this material is strictly forbidden unless prior written permission is
-// obtained from COMPANY. Access to the source code contained herein is
-// hereby forbidden to anyone except current COMPANY employees, managers
-// or contractors who have executed Confidentiality and Non-disclosure
-// agreements explicitly covering such access.
-//
-// The copyright notice above does not evidence any actual or intended
-// publication or disclosure of this source code, which includes
-// information that is confidential and/or proprietary, and is a trade
-// secret, of COMPANY. ANY REPRODUCTION, MODIFICATION, DISTRIBUTION,
-// PUBLIC PERFORMANCE, OR PUBLIC DISPLAY OF OR THROUGH USE OF THIS
-// SOURCE CODE WITHOUT THE EXPRESS WRITTEN CONSENT OF COMPANY IS
-// STRICTLY PROHIBITED, AND IN VIOLATION OF APPLICABLE LAWS AND
-// INTERNATIONAL TREATIES. THE RECEIPT OR POSSESSION OF THIS SOURCE
-// CODE AND/OR RELATED INFORMATION DOES NOT CONVEY OR IMPLY ANY RIGHTS
-// TO REPRODUCE, DISCLOSE OR DISTRIBUTE ITS CONTENTS, OR TO MANUFACTURE,
-// USE, OR SELL ANYTHING THAT IT MAY DESCRIBE, IN WHOLE OR IN PART.
-//
-// %COPYRIGHT_END%
-// --------------------------------------------------------------------
-// %BANNER_END%
+// Copyright 1998-2018 Epic Games, Inc. All Rights Reserved.
 
 #pragma once
 
diff --git a/Engine/Plugins/Lumin/MagicLeap/Source/MagicLeapGestures/Public/MagicLeapGestureTypes.h b/Engine/Plugins/Lumin/MagicLeap/Source/MagicLeapGestures/Public/MagicLeapGestureTypes.h
--- a/Engine/Plugins/Lumin/MagicLeap/Source/MagicLeapGestures/Public/MagicLeapGestureTypes.h
+++ b/Engine/Plugins/Lumin/MagicLeap/Source/MagicLeapGestures/Public/MagicLeapGestureTypes.h
@@ -1,36 +1,4 @@
-// %BANNER_BEGIN%
-// ---------------------------------------------------------------------
-// %COPYRIGHT_BEGIN%
-//
-// Copyright (c) 2017 Magic Leap, Inc. (COMPANY) All Rights Reserved.
-// Magic Leap, Inc. Confidential and Proprietary
-//
-// NOTICE: All information contained herein is, and remains the property
-// of COMPANY. The intellectual and technical concepts contained herein
-// are proprietary to COMPANY and may be covered by U.S. and Foreign
-// Patents, patents in process, and are protected by trade secret or
-// copyright law. Dissemination of this information or reproduction of
-// this material is strictly forbidden unless prior written permission is
-// obtained from COMPANY. Access to the source code contained herein is
-// hereby forbidden to anyone except current COMPANY employees, managers
-// or contractors who have executed Confidentiality and Non-disclosure
-// agreements explicitly covering such access.
-//
-// The copyright notice above does not evidence any actual or intended
-// publication or disclosure of this source code, which includes
-// information that is confidential and/or proprietary, and is a trade
-// secret, of COMPANY. ANY REPRODUCTION, MODIFICATION, DISTRIBUTION,
-// PUBLIC PERFORMANCE, OR PUBLIC DISPLAY OF OR THROUGH USE OF THIS
-// SOURCE CODE WITHOUT THE EXPRESS WRITTEN CONSENT OF COMPANY IS
-// STRICTLY PROHIBITED, AND IN VIOLATION OF APPLICABLE LAWS AND
-// INTERNATIONAL TREATIES. THE RECEIPT OR POSSESSION OF THIS SOURCE
-// CODE AND/OR RELATED INFORMATION DOES NOT CONVEY OR IMPLY ANY RIGHTS
-// TO REPRODUCE, DISCLOSE OR DISTRIBUTE ITS CONTENTS, OR TO MANUFACTURE,
-// USE, OR SELL ANYTHING THAT IT MAY DESCRIBE, IN WHOLE OR IN PART.
-//
-// %COPYRIGHT_END%
-// --------------------------------------------------------------------
-// %BANNER_END%
+// Copyright 1998-2018 Epic Games, Inc. All Rights Reserved.
 
 #pragma once
 
diff --git a/Engine/Plugins/Lumin/MagicLeap/Source/MagicLeapGestures/Public/MagicLeapGesturesFunctionLibrary.h b/Engine/Plugins/Lumin/MagicLeap/Source/MagicLeapGestures/Public/MagicLeapGesturesFunctionLibrary.h
--- a/Engine/Plugins/Lumin/MagicLeap/Source/MagicLeapGestures/Public/MagicLeapGesturesFunctionLibrary.h
+++ b/Engine/Plugins/Lumin/MagicLeap/Source/MagicLeapGestures/Public/MagicLeapGesturesFunctionLibrary.h
@@ -1,36 +1,4 @@
-// %BANNER_BEGIN%
-// ---------------------------------------------------------------------
-// %COPYRIGHT_BEGIN%
-//
-// Copyright (c) 2017 Magic Leap, Inc. (COMPANY) All Rights Reserved.
-// Magic Leap, Inc. Confidential and Proprietary
-//
-// NOTICE: All information contained herein is, and remains the property
-// of COMPANY. The intellectual and technical concepts contained herein
-// are proprietary to COMPANY and may be covered by U.S. and Foreign
-// Patents, patents in process, and are protected by trade secret or
-// copyright law. Dissemination of this information or reproduction of
-// this material is strictly forbidden unless prior written permission is
-// obtained from COMPANY. Access to the source code contained herein is
-// hereby forbidden to anyone except current COMPANY employees, managers
-// or contractors who have executed Confidentiality and Non-disclosure
-// agreements explicitly covering such access.
-//
-// The copyright notice above does not evidence any actual or intended
-// publication or disclosure of this source code, which includes
-// information that is confidential and/or proprietary, and is a trade
-// secret, of COMPANY. ANY REPRODUCTION, MODIFICATION, DISTRIBUTION,
-// PUBLIC PERFORMANCE, OR PUBLIC DISPLAY OF OR THROUGH USE OF THIS
-// SOURCE CODE WITHOUT THE EXPRESS WRITTEN CONSENT OF COMPANY IS
-// STRICTLY PROHIBITED, AND IN VIOLATION OF APPLICABLE LAWS AND
-// INTERNATIONAL TREATIES. THE RECEIPT OR POSSESSION OF THIS SOURCE
-// CODE AND/OR RELATED INFORMATION DOES NOT CONVEY OR IMPLY ANY RIGHTS
-// TO REPRODUCE, DISCLOSE OR DISTRIBUTE ITS CONTENTS, OR TO MANUFACTURE,
-// USE, OR SELL ANYTHING THAT IT MAY DESCRIBE, IN WHOLE OR IN PART.
-//
-// %COPYRIGHT_END%
-// --------------------------------------------------------------------
-// %BANNER_END%
+// Copyright 1998-2018 Epic Games, Inc. All Rights Reserved.
 
 #pragma once
 
diff --git a/Engine/Plugins/Lumin/MagicLeap/Source/MagicLeapHelperOpenGL/MagicLeapHelperOpenGL.Build.cs b/Engine/Plugins/Lumin/MagicLeap/Source/MagicLeapHelperOpenGL/MagicLeapHelperOpenGL.Build.cs
--- a/Engine/Plugins/Lumin/MagicLeap/Source/MagicLeapHelperOpenGL/MagicLeapHelperOpenGL.Build.cs
+++ b/Engine/Plugins/Lumin/MagicLeap/Source/MagicLeapHelperOpenGL/MagicLeapHelperOpenGL.Build.cs
@@ -1,36 +1,4 @@
-// %BANNER_BEGIN%
-// ---------------------------------------------------------------------
-// %COPYRIGHT_BEGIN%
-//
-// Copyright (c) 2017 Magic Leap, Inc. (COMPANY) All Rights Reserved.
-// Magic Leap, Inc. Confidential and Proprietary
-//
-// NOTICE: All information contained herein is, and remains the property
-// of COMPANY. The intellectual and technical concepts contained herein
-// are proprietary to COMPANY and may be covered by U.S. and Foreign
-// Patents, patents in process, and are protected by trade secret or
-// copyright law. Dissemination of this information or reproduction of
-// this material is strictly forbidden unless prior written permission is
-// obtained from COMPANY. Access to the source code contained herein is
-// hereby forbidden to anyone except current COMPANY employees, managers
-// or contractors who have executed Confidentiality and Non-disclosure
-// agreements explicitly covering such access.
-//
-// The copyright notice above does not evidence any actual or intended
-// publication or disclosure of this source code, which includes
-// information that is confidential and/or proprietary, and is a trade
-// secret, of COMPANY. ANY REPRODUCTION, MODIFICATION, DISTRIBUTION,
-// PUBLIC PERFORMANCE, OR PUBLIC DISPLAY OF OR THROUGH USE OF THIS
-// SOURCE CODE WITHOUT THE EXPRESS WRITTEN CONSENT OF COMPANY IS
-// STRICTLY PROHIBITED, AND IN VIOLATION OF APPLICABLE LAWS AND
-// INTERNATIONAL TREATIES. THE RECEIPT OR POSSESSION OF THIS SOURCE
-// CODE AND/OR RELATED INFORMATION DOES NOT CONVEY OR IMPLY ANY RIGHTS
-// TO REPRODUCE, DISCLOSE OR DISTRIBUTE ITS CONTENTS, OR TO MANUFACTURE,
-// USE, OR SELL ANYTHING THAT IT MAY DESCRIBE, IN WHOLE OR IN PART.
-//
-// %COPYRIGHT_END%
-// --------------------------------------------------------------------
-// %BANNER_END%
+// Copyright 1998-2018 Epic Games, Inc. All Rights Reserved.
 
 using System.IO;
 
diff --git a/Engine/Plugins/Lumin/MagicLeap/Source/MagicLeapHelperOpenGL/Private/MagicLeapHelperOpenGL.cpp b/Engine/Plugins/Lumin/MagicLeap/Source/MagicLeapHelperOpenGL/Private/MagicLeapHelperOpenGL.cpp
--- a/Engine/Plugins/Lumin/MagicLeap/Source/MagicLeapHelperOpenGL/Private/MagicLeapHelperOpenGL.cpp
+++ b/Engine/Plugins/Lumin/MagicLeap/Source/MagicLeapHelperOpenGL/Private/MagicLeapHelperOpenGL.cpp
@@ -1,36 +1,4 @@
-// %BANNER_BEGIN%
-// ---------------------------------------------------------------------
-// %COPYRIGHT_BEGIN%
-//
-// Copyright (c) 2017 Magic Leap, Inc. (COMPANY) All Rights Reserved.
-// Magic Leap, Inc. Confidential and Proprietary
-//
-// NOTICE: All information contained herein is, and remains the property
-// of COMPANY. The intellectual and technical concepts contained herein
-// are proprietary to COMPANY and may be covered by U.S. and Foreign
-// Patents, patents in process, and are protected by trade secret or
-// copyright law. Dissemination of this information or reproduction of
-// this material is strictly forbidden unless prior written permission is
-// obtained from COMPANY. Access to the source code contained herein is
-// hereby forbidden to anyone except current COMPANY employees, managers
-// or contractors who have executed Confidentiality and Non-disclosure
-// agreements explicitly covering such access.
-//
-// The copyright notice above does not evidence any actual or intended
-// publication or disclosure of this source code, which includes
-// information that is confidential and/or proprietary, and is a trade
-// secret, of COMPANY. ANY REPRODUCTION, MODIFICATION, DISTRIBUTION,
-// PUBLIC PERFORMANCE, OR PUBLIC DISPLAY OF OR THROUGH USE OF THIS
-// SOURCE CODE WITHOUT THE EXPRESS WRITTEN CONSENT OF COMPANY IS
-// STRICTLY PROHIBITED, AND IN VIOLATION OF APPLICABLE LAWS AND
-// INTERNATIONAL TREATIES. THE RECEIPT OR POSSESSION OF THIS SOURCE
-// CODE AND/OR RELATED INFORMATION DOES NOT CONVEY OR IMPLY ANY RIGHTS
-// TO REPRODUCE, DISCLOSE OR DISTRIBUTE ITS CONTENTS, OR TO MANUFACTURE,
-// USE, OR SELL ANYTHING THAT IT MAY DESCRIBE, IN WHOLE OR IN PART.
-//
-// %COPYRIGHT_END%
-// --------------------------------------------------------------------
-// %BANNER_END%
+// Copyright 1998-2018 Epic Games, Inc. All Rights Reserved.
 
 #include "MagicLeapHelperOpenGL.h"
 #include "IMagicLeapHelperOpenGLPlugin.h"
diff --git a/Engine/Plugins/Lumin/MagicLeap/Source/MagicLeapHelperOpenGL/Private/TexturePixelReader.cpp b/Engine/Plugins/Lumin/MagicLeap/Source/MagicLeapHelperOpenGL/Private/TexturePixelReader.cpp
--- a/Engine/Plugins/Lumin/MagicLeap/Source/MagicLeapHelperOpenGL/Private/TexturePixelReader.cpp
+++ b/Engine/Plugins/Lumin/MagicLeap/Source/MagicLeapHelperOpenGL/Private/TexturePixelReader.cpp
@@ -1,36 +1,4 @@
-// %BANNER_BEGIN%
-// ---------------------------------------------------------------------
-// %COPYRIGHT_BEGIN%
-//
-// Copyright (c) 2017 Magic Leap, Inc. (COMPANY) All Rights Reserved.
-// Magic Leap, Inc. Confidential and Proprietary
-//
-// NOTICE:  All information contained herein is, and remains the property
-// of COMPANY. The intellectual and technical concepts contained herein
-// are proprietary to COMPANY and may be covered by U.S. and Foreign
-// Patents, patents in process, and are protected by trade secret or
-// copyright law.  Dissemination of this information or reproduction of
-// this material is strictly forbidden unless prior written permission is
-// obtained from COMPANY.  Access to the source code contained herein is
-// hereby forbidden to anyone except current COMPANY employees, managers
-// or contractors who have executed Confidentiality and Non-disclosure
-// agreements explicitly covering such access.
-//
-// The copyright notice above does not evidence any actual or intended
-// publication or disclosure  of  this source code, which includes
-// information that is confidential and/or proprietary, and is a trade
-// secret, of  COMPANY.   ANY REPRODUCTION, MODIFICATION, DISTRIBUTION,
-// PUBLIC  PERFORMANCE, OR PUBLIC DISPLAY OF OR THROUGH USE  OF THIS
-// SOURCE CODE  WITHOUT THE EXPRESS WRITTEN CONSENT OF COMPANY IS
-// STRICTLY PROHIBITED, AND IN VIOLATION OF APPLICABLE LAWS AND
-// INTERNATIONAL TREATIES.  THE RECEIPT OR POSSESSION OF  THIS SOURCE
-// CODE AND/OR RELATED INFORMATION DOES NOT CONVEY OR IMPLY ANY RIGHTS
-// TO REPRODUCE, DISCLOSE OR DISTRIBUTE ITS CONTENTS, OR TO MANUFACTURE,
-// USE, OR SELL ANYTHING THAT IT  MAY DESCRIBE, IN WHOLE OR IN PART.
-//
-// %COPYRIGHT_END%
-// --------------------------------------------------------------------*/
-// %BANNER_END%
+// Copyright 1998-2018 Epic Games, Inc. All Rights Reserved.
 
 #include "TexturePixelReader.h"
 #include "Misc/CString.h"
diff --git a/Engine/Plugins/Lumin/MagicLeap/Source/MagicLeapHelperOpenGL/Public/IMagicLeapHelperOpenGLPlugin.h b/Engine/Plugins/Lumin/MagicLeap/Source/MagicLeapHelperOpenGL/Public/IMagicLeapHelperOpenGLPlugin.h
--- a/Engine/Plugins/Lumin/MagicLeap/Source/MagicLeapHelperOpenGL/Public/IMagicLeapHelperOpenGLPlugin.h
+++ b/Engine/Plugins/Lumin/MagicLeap/Source/MagicLeapHelperOpenGL/Public/IMagicLeapHelperOpenGLPlugin.h
@@ -1,36 +1,4 @@
-// %BANNER_BEGIN%
-// ---------------------------------------------------------------------
-// %COPYRIGHT_BEGIN%
-//
-// Copyright (c) 2017 Magic Leap, Inc. (COMPANY) All Rights Reserved.
-// Magic Leap, Inc. Confidential and Proprietary
-//
-// NOTICE: All information contained herein is, and remains the property
-// of COMPANY. The intellectual and technical concepts contained herein
-// are proprietary to COMPANY and may be covered by U.S. and Foreign
-// Patents, patents in process, and are protected by trade secret or
-// copyright law. Dissemination of this information or reproduction of
-// this material is strictly forbidden unless prior written permission is
-// obtained from COMPANY. Access to the source code contained herein is
-// hereby forbidden to anyone except current COMPANY employees, managers
-// or contractors who have executed Confidentiality and Non-disclosure
-// agreements explicitly covering such access.
-//
-// The copyright notice above does not evidence any actual or intended
-// publication or disclosure of this source code, which includes
-// information that is confidential and/or proprietary, and is a trade
-// secret, of COMPANY. ANY REPRODUCTION, MODIFICATION, DISTRIBUTION,
-// PUBLIC PERFORMANCE, OR PUBLIC DISPLAY OF OR THROUGH USE OF THIS
-// SOURCE CODE WITHOUT THE EXPRESS WRITTEN CONSENT OF COMPANY IS
-// STRICTLY PROHIBITED, AND IN VIOLATION OF APPLICABLE LAWS AND
-// INTERNATIONAL TREATIES. THE RECEIPT OR POSSESSION OF THIS SOURCE
-// CODE AND/OR RELATED INFORMATION DOES NOT CONVEY OR IMPLY ANY RIGHTS
-// TO REPRODUCE, DISCLOSE OR DISTRIBUTE ITS CONTENTS, OR TO MANUFACTURE,
-// USE, OR SELL ANYTHING THAT IT MAY DESCRIBE, IN WHOLE OR IN PART.
-//
-// %COPYRIGHT_END%
-// --------------------------------------------------------------------
-// %BANNER_END%
+// Copyright 1998-2018 Epic Games, Inc. All Rights Reserved.
 
 #pragma once
 
diff --git a/Engine/Plugins/Lumin/MagicLeap/Source/MagicLeapHelperOpenGL/Public/MagicLeapHelperOpenGL.h b/Engine/Plugins/Lumin/MagicLeap/Source/MagicLeapHelperOpenGL/Public/MagicLeapHelperOpenGL.h
--- a/Engine/Plugins/Lumin/MagicLeap/Source/MagicLeapHelperOpenGL/Public/MagicLeapHelperOpenGL.h
+++ b/Engine/Plugins/Lumin/MagicLeap/Source/MagicLeapHelperOpenGL/Public/MagicLeapHelperOpenGL.h
@@ -1,36 +1,4 @@
-// %BANNER_BEGIN%
-// ---------------------------------------------------------------------
-// %COPYRIGHT_BEGIN%
-//
-// Copyright (c) 2017 Magic Leap, Inc. (COMPANY) All Rights Reserved.
-// Magic Leap, Inc. Confidential and Proprietary
-//
-// NOTICE: All information contained herein is, and remains the property
-// of COMPANY. The intellectual and technical concepts contained herein
-// are proprietary to COMPANY and may be covered by U.S. and Foreign
-// Patents, patents in process, and are protected by trade secret or
-// copyright law. Dissemination of this information or reproduction of
-// this material is strictly forbidden unless prior written permission is
-// obtained from COMPANY. Access to the source code contained herein is
-// hereby forbidden to anyone except current COMPANY employees, managers
-// or contractors who have executed Confidentiality and Non-disclosure
-// agreements explicitly covering such access.
-//
-// The copyright notice above does not evidence any actual or intended
-// publication or disclosure of this source code, which includes
-// information that is confidential and/or proprietary, and is a trade
-// secret, of COMPANY. ANY REPRODUCTION, MODIFICATION, DISTRIBUTION,
-// PUBLIC PERFORMANCE, OR PUBLIC DISPLAY OF OR THROUGH USE OF THIS
-// SOURCE CODE WITHOUT THE EXPRESS WRITTEN CONSENT OF COMPANY IS
-// STRICTLY PROHIBITED, AND IN VIOLATION OF APPLICABLE LAWS AND
-// INTERNATIONAL TREATIES. THE RECEIPT OR POSSESSION OF THIS SOURCE
-// CODE AND/OR RELATED INFORMATION DOES NOT CONVEY OR IMPLY ANY RIGHTS
-// TO REPRODUCE, DISCLOSE OR DISTRIBUTE ITS CONTENTS, OR TO MANUFACTURE,
-// USE, OR SELL ANYTHING THAT IT MAY DESCRIBE, IN WHOLE OR IN PART.
-//
-// %COPYRIGHT_END%
-// --------------------------------------------------------------------
-// %BANNER_END%
+// Copyright 1998-2018 Epic Games, Inc. All Rights Reserved.
 
 #pragma once
 
diff --git a/Engine/Plugins/Lumin/MagicLeap/Source/MagicLeapHelperOpenGL/Public/TexturePixelReader.h b/Engine/Plugins/Lumin/MagicLeap/Source/MagicLeapHelperOpenGL/Public/TexturePixelReader.h
--- a/Engine/Plugins/Lumin/MagicLeap/Source/MagicLeapHelperOpenGL/Public/TexturePixelReader.h
+++ b/Engine/Plugins/Lumin/MagicLeap/Source/MagicLeapHelperOpenGL/Public/TexturePixelReader.h
@@ -1,36 +1,4 @@
-// %BANNER_BEGIN%
-// ---------------------------------------------------------------------
-// %COPYRIGHT_BEGIN%
-//
-// Copyright (c) 2017 Magic Leap, Inc. (COMPANY) All Rights Reserved.
-// Magic Leap, Inc. Confidential and Proprietary
-//
-// NOTICE:  All information contained herein is, and remains the property
-// of COMPANY. The intellectual and technical concepts contained herein
-// are proprietary to COMPANY and may be covered by U.S. and Foreign
-// Patents, patents in process, and are protected by trade secret or
-// copyright law.  Dissemination of this information or reproduction of
-// this material is strictly forbidden unless prior written permission is
-// obtained from COMPANY.  Access to the source code contained herein is
-// hereby forbidden to anyone except current COMPANY employees, managers
-// or contractors who have executed Confidentiality and Non-disclosure
-// agreements explicitly covering such access.
-//
-// The copyright notice above does not evidence any actual or intended
-// publication or disclosure  of  this source code, which includes
-// information that is confidential and/or proprietary, and is a trade
-// secret, of  COMPANY.   ANY REPRODUCTION, MODIFICATION, DISTRIBUTION,
-// PUBLIC  PERFORMANCE, OR PUBLIC DISPLAY OF OR THROUGH USE  OF THIS
-// SOURCE CODE  WITHOUT THE EXPRESS WRITTEN CONSENT OF COMPANY IS
-// STRICTLY PROHIBITED, AND IN VIOLATION OF APPLICABLE LAWS AND
-// INTERNATIONAL TREATIES.  THE RECEIPT OR POSSESSION OF  THIS SOURCE
-// CODE AND/OR RELATED INFORMATION DOES NOT CONVEY OR IMPLY ANY RIGHTS
-// TO REPRODUCE, DISCLOSE OR DISTRIBUTE ITS CONTENTS, OR TO MANUFACTURE,
-// USE, OR SELL ANYTHING THAT IT  MAY DESCRIBE, IN WHOLE OR IN PART.
-//
-// %COPYRIGHT_END%
-// --------------------------------------------------------------------*/
-// %BANNER_END%
+// Copyright 1998-2018 Epic Games, Inc. All Rights Reserved.
 
 #pragma once
 
diff --git a/Engine/Plugins/Lumin/MagicLeap/Source/MagicLeapHelperVulkan/MagicLeapHelperVulkan.Build.cs b/Engine/Plugins/Lumin/MagicLeap/Source/MagicLeapHelperVulkan/MagicLeapHelperVulkan.Build.cs
--- a/Engine/Plugins/Lumin/MagicLeap/Source/MagicLeapHelperVulkan/MagicLeapHelperVulkan.Build.cs
+++ b/Engine/Plugins/Lumin/MagicLeap/Source/MagicLeapHelperVulkan/MagicLeapHelperVulkan.Build.cs
@@ -1,36 +1,4 @@
-// %BANNER_BEGIN%
-// ---------------------------------------------------------------------
-// %COPYRIGHT_BEGIN%
-//
-// Copyright (c) 2017 Magic Leap, Inc. (COMPANY) All Rights Reserved.
-// Magic Leap, Inc. Confidential and Proprietary
-//
-// NOTICE: All information contained herein is, and remains the property
-// of COMPANY. The intellectual and technical concepts contained herein
-// are proprietary to COMPANY and may be covered by U.S. and Foreign
-// Patents, patents in process, and are protected by trade secret or
-// copyright law. Dissemination of this information or reproduction of
-// this material is strictly forbidden unless prior written permission is
-// obtained from COMPANY. Access to the source code contained herein is
-// hereby forbidden to anyone except current COMPANY employees, managers
-// or contractors who have executed Confidentiality and Non-disclosure
-// agreements explicitly covering such access.
-//
-// The copyright notice above does not evidence any actual or intended
-// publication or disclosure of this source code, which includes
-// information that is confidential and/or proprietary, and is a trade
-// secret, of COMPANY. ANY REPRODUCTION, MODIFICATION, DISTRIBUTION,
-// PUBLIC PERFORMANCE, OR PUBLIC DISPLAY OF OR THROUGH USE OF THIS
-// SOURCE CODE WITHOUT THE EXPRESS WRITTEN CONSENT OF COMPANY IS
-// STRICTLY PROHIBITED, AND IN VIOLATION OF APPLICABLE LAWS AND
-// INTERNATIONAL TREATIES. THE RECEIPT OR POSSESSION OF THIS SOURCE
-// CODE AND/OR RELATED INFORMATION DOES NOT CONVEY OR IMPLY ANY RIGHTS
-// TO REPRODUCE, DISCLOSE OR DISTRIBUTE ITS CONTENTS, OR TO MANUFACTURE,
-// USE, OR SELL ANYTHING THAT IT MAY DESCRIBE, IN WHOLE OR IN PART.
-//
-// %COPYRIGHT_END%
-// --------------------------------------------------------------------
-// %BANNER_END%
+// Copyright 1998-2018 Epic Games, Inc. All Rights Reserved.
 
 using System.IO;
 
diff --git a/Engine/Plugins/Lumin/MagicLeap/Source/MagicLeapHelperVulkan/Private/MagicLeapHelperVulkan.cpp b/Engine/Plugins/Lumin/MagicLeap/Source/MagicLeapHelperVulkan/Private/MagicLeapHelperVulkan.cpp
--- a/Engine/Plugins/Lumin/MagicLeap/Source/MagicLeapHelperVulkan/Private/MagicLeapHelperVulkan.cpp
+++ b/Engine/Plugins/Lumin/MagicLeap/Source/MagicLeapHelperVulkan/Private/MagicLeapHelperVulkan.cpp
@@ -1,36 +1,4 @@
-// %BANNER_BEGIN%
-// ---------------------------------------------------------------------
-// %COPYRIGHT_BEGIN%
-//
-// Copyright (c) 2017 Magic Leap, Inc. (COMPANY) All Rights Reserved.
-// Magic Leap, Inc. Confidential and Proprietary
-//
-// NOTICE: All information contained herein is, and remains the property
-// of COMPANY. The intellectual and technical concepts contained herein
-// are proprietary to COMPANY and may be covered by U.S. and Foreign
-// Patents, patents in process, and are protected by trade secret or
-// copyright law. Dissemination of this information or reproduction of
-// this material is strictly forbidden unless prior written permission is
-// obtained from COMPANY. Access to the source code contained herein is
-// hereby forbidden to anyone except current COMPANY employees, managers
-// or contractors who have executed Confidentiality and Non-disclosure
-// agreements explicitly covering such access.
-//
-// The copyright notice above does not evidence any actual or intended
-// publication or disclosure of this source code, which includes
-// information that is confidential and/or proprietary, and is a trade
-// secret, of COMPANY. ANY REPRODUCTION, MODIFICATION, DISTRIBUTION,
-// PUBLIC PERFORMANCE, OR PUBLIC DISPLAY OF OR THROUGH USE OF THIS
-// SOURCE CODE WITHOUT THE EXPRESS WRITTEN CONSENT OF COMPANY IS
-// STRICTLY PROHIBITED, AND IN VIOLATION OF APPLICABLE LAWS AND
-// INTERNATIONAL TREATIES. THE RECEIPT OR POSSESSION OF THIS SOURCE
-// CODE AND/OR RELATED INFORMATION DOES NOT CONVEY OR IMPLY ANY RIGHTS
-// TO REPRODUCE, DISCLOSE OR DISTRIBUTE ITS CONTENTS, OR TO MANUFACTURE,
-// USE, OR SELL ANYTHING THAT IT MAY DESCRIBE, IN WHOLE OR IN PART.
-//
-// %COPYRIGHT_END%
-// --------------------------------------------------------------------
-// %BANNER_END%
+// Copyright 1998-2018 Epic Games, Inc. All Rights Reserved.
 
 #include "MagicLeapHelperVulkan.h"
 #include "IMagicLeapHelperVulkanPlugin.h"
diff --git a/Engine/Plugins/Lumin/MagicLeap/Source/MagicLeapHelperVulkan/Public/IMagicLeapHelperVulkanPlugin.h b/Engine/Plugins/Lumin/MagicLeap/Source/MagicLeapHelperVulkan/Public/IMagicLeapHelperVulkanPlugin.h
--- a/Engine/Plugins/Lumin/MagicLeap/Source/MagicLeapHelperVulkan/Public/IMagicLeapHelperVulkanPlugin.h
+++ b/Engine/Plugins/Lumin/MagicLeap/Source/MagicLeapHelperVulkan/Public/IMagicLeapHelperVulkanPlugin.h
@@ -1,36 +1,4 @@
-// %BANNER_BEGIN%
-// ---------------------------------------------------------------------
-// %COPYRIGHT_BEGIN%
-//
-// Copyright (c) 2017 Magic Leap, Inc. (COMPANY) All Rights Reserved.
-// Magic Leap, Inc. Confidential and Proprietary
-//
-// NOTICE: All information contained herein is, and remains the property
-// of COMPANY. The intellectual and technical concepts contained herein
-// are proprietary to COMPANY and may be covered by U.S. and Foreign
-// Patents, patents in process, and are protected by trade secret or
-// copyright law. Dissemination of this information or reproduction of
-// this material is strictly forbidden unless prior written permission is
-// obtained from COMPANY. Access to the source code contained herein is
-// hereby forbidden to anyone except current COMPANY employees, managers
-// or contractors who have executed Confidentiality and Non-disclosure
-// agreements explicitly covering such access.
-//
-// The copyright notice above does not evidence any actual or intended
-// publication or disclosure of this source code, which includes
-// information that is confidential and/or proprietary, and is a trade
-// secret, of COMPANY. ANY REPRODUCTION, MODIFICATION, DISTRIBUTION,
-// PUBLIC PERFORMANCE, OR PUBLIC DISPLAY OF OR THROUGH USE OF THIS
-// SOURCE CODE WITHOUT THE EXPRESS WRITTEN CONSENT OF COMPANY IS
-// STRICTLY PROHIBITED, AND IN VIOLATION OF APPLICABLE LAWS AND
-// INTERNATIONAL TREATIES. THE RECEIPT OR POSSESSION OF THIS SOURCE
-// CODE AND/OR RELATED INFORMATION DOES NOT CONVEY OR IMPLY ANY RIGHTS
-// TO REPRODUCE, DISCLOSE OR DISTRIBUTE ITS CONTENTS, OR TO MANUFACTURE,
-// USE, OR SELL ANYTHING THAT IT MAY DESCRIBE, IN WHOLE OR IN PART.
-//
-// %COPYRIGHT_END%
-// --------------------------------------------------------------------
-// %BANNER_END%
+// Copyright 1998-2018 Epic Games, Inc. All Rights Reserved.
 
 #pragma once
 
diff --git a/Engine/Plugins/Lumin/MagicLeap/Source/MagicLeapHelperVulkan/Public/MagicLeapHelperVulkan.h b/Engine/Plugins/Lumin/MagicLeap/Source/MagicLeapHelperVulkan/Public/MagicLeapHelperVulkan.h
--- a/Engine/Plugins/Lumin/MagicLeap/Source/MagicLeapHelperVulkan/Public/MagicLeapHelperVulkan.h
+++ b/Engine/Plugins/Lumin/MagicLeap/Source/MagicLeapHelperVulkan/Public/MagicLeapHelperVulkan.h
@@ -1,36 +1,4 @@
-// %BANNER_BEGIN%
-// ---------------------------------------------------------------------
-// %COPYRIGHT_BEGIN%
-//
-// Copyright (c) 2017 Magic Leap, Inc. (COMPANY) All Rights Reserved.
-// Magic Leap, Inc. Confidential and Proprietary
-//
-// NOTICE: All information contained herein is, and remains the property
-// of COMPANY. The intellectual and technical concepts contained herein
-// are proprietary to COMPANY and may be covered by U.S. and Foreign
-// Patents, patents in process, and are protected by trade secret or
-// copyright law. Dissemination of this information or reproduction of
-// this material is strictly forbidden unless prior written permission is
-// obtained from COMPANY. Access to the source code contained herein is
-// hereby forbidden to anyone except current COMPANY employees, managers
-// or contractors who have executed Confidentiality and Non-disclosure
-// agreements explicitly covering such access.
-//
-// The copyright notice above does not evidence any actual or intended
-// publication or disclosure of this source code, which includes
-// information that is confidential and/or proprietary, and is a trade
-// secret, of COMPANY. ANY REPRODUCTION, MODIFICATION, DISTRIBUTION,
-// PUBLIC PERFORMANCE, OR PUBLIC DISPLAY OF OR THROUGH USE OF THIS
-// SOURCE CODE WITHOUT THE EXPRESS WRITTEN CONSENT OF COMPANY IS
-// STRICTLY PROHIBITED, AND IN VIOLATION OF APPLICABLE LAWS AND
-// INTERNATIONAL TREATIES. THE RECEIPT OR POSSESSION OF THIS SOURCE
-// CODE AND/OR RELATED INFORMATION DOES NOT CONVEY OR IMPLY ANY RIGHTS
-// TO REPRODUCE, DISCLOSE OR DISTRIBUTE ITS CONTENTS, OR TO MANUFACTURE,
-// USE, OR SELL ANYTHING THAT IT MAY DESCRIBE, IN WHOLE OR IN PART.
-//
-// %COPYRIGHT_END%
-// --------------------------------------------------------------------
-// %BANNER_END%
+// Copyright 1998-2018 Epic Games, Inc. All Rights Reserved.
 
 #pragma once
 
diff --git a/Engine/Plugins/Lumin/MagicLeap/Source/MagicLeapIdentity/MagicLeapIdentity.Build.cs b/Engine/Plugins/Lumin/MagicLeap/Source/MagicLeapIdentity/MagicLeapIdentity.Build.cs
--- a/Engine/Plugins/Lumin/MagicLeap/Source/MagicLeapIdentity/MagicLeapIdentity.Build.cs
+++ b/Engine/Plugins/Lumin/MagicLeap/Source/MagicLeapIdentity/MagicLeapIdentity.Build.cs
@@ -1,36 +1,4 @@
-// %BANNER_BEGIN%
-// ---------------------------------------------------------------------
-// %COPYRIGHT_BEGIN%
-//
-// Copyright (c) 2017 Magic Leap, Inc. (COMPANY) All Rights Reserved.
-// Magic Leap, Inc. Confidential and Proprietary
-//
-// NOTICE: All information contained herein is, and remains the property
-// of COMPANY. The intellectual and technical concepts contained herein
-// are proprietary to COMPANY and may be covered by U.S. and Foreign
-// Patents, patents in process, and are protected by trade secret or
-// copyright law. Dissemination of this information or reproduction of
-// this material is strictly forbidden unless prior written permission is
-// obtained from COMPANY. Access to the source code contained herein is
-// hereby forbidden to anyone except current COMPANY employees, managers
-// or contractors who have executed Confidentiality and Non-disclosure
-// agreements explicitly covering such access.
-//
-// The copyright notice above does not evidence any actual or intended
-// publication or disclosure of this source code, which includes
-// information that is confidential and/or proprietary, and is a trade
-// secret, of COMPANY. ANY REPRODUCTION, MODIFICATION, DISTRIBUTION,
-// PUBLIC PERFORMANCE, OR PUBLIC DISPLAY OF OR THROUGH USE OF THIS
-// SOURCE CODE WITHOUT THE EXPRESS WRITTEN CONSENT OF COMPANY IS
-// STRICTLY PROHIBITED, AND IN VIOLATION OF APPLICABLE LAWS AND
-// INTERNATIONAL TREATIES. THE RECEIPT OR POSSESSION OF THIS SOURCE
-// CODE AND/OR RELATED INFORMATION DOES NOT CONVEY OR IMPLY ANY RIGHTS
-// TO REPRODUCE, DISCLOSE OR DISTRIBUTE ITS CONTENTS, OR TO MANUFACTURE,
-// USE, OR SELL ANYTHING THAT IT MAY DESCRIBE, IN WHOLE OR IN PART.
-//
-// %COPYRIGHT_END%
-// --------------------------------------------------------------------
-// %BANNER_END%
+// Copyright 1998-2018 Epic Games, Inc. All Rights Reserved.
 
 using System.IO;
 
diff --git a/Engine/Plugins/Lumin/MagicLeap/Source/MagicLeapIdentity/Private/MagicLeapIdentity.cpp b/Engine/Plugins/Lumin/MagicLeap/Source/MagicLeapIdentity/Private/MagicLeapIdentity.cpp
--- a/Engine/Plugins/Lumin/MagicLeap/Source/MagicLeapIdentity/Private/MagicLeapIdentity.cpp
+++ b/Engine/Plugins/Lumin/MagicLeap/Source/MagicLeapIdentity/Private/MagicLeapIdentity.cpp
@@ -1,36 +1,4 @@
-// %BANNER_BEGIN%
-// ---------------------------------------------------------------------
-// %COPYRIGHT_BEGIN%
-//
-// Copyright (c) 2017 Magic Leap, Inc. (COMPANY) All Rights Reserved.
-// Magic Leap, Inc. Confidential and Proprietary
-//
-// NOTICE: All information contained herein is, and remains the property
-// of COMPANY. The intellectual and technical concepts contained herein
-// are proprietary to COMPANY and may be covered by U.S. and Foreign
-// Patents, patents in process, and are protected by trade secret or
-// copyright law. Dissemination of this information or reproduction of
-// this material is strictly forbidden unless prior written permission is
-// obtained from COMPANY. Access to the source code contained herein is
-// hereby forbidden to anyone except current COMPANY employees, managers
-// or contractors who have executed Confidentiality and Non-disclosure
-// agreements explicitly covering such access.
-//
-// The copyright notice above does not evidence any actual or intended
-// publication or disclosure of this source code, which includes
-// information that is confidential and/or proprietary, and is a trade
-// secret, of COMPANY. ANY REPRODUCTION, MODIFICATION, DISTRIBUTION,
-// PUBLIC PERFORMANCE, OR PUBLIC DISPLAY OF OR THROUGH USE OF THIS
-// SOURCE CODE WITHOUT THE EXPRESS WRITTEN CONSENT OF COMPANY IS
-// STRICTLY PROHIBITED, AND IN VIOLATION OF APPLICABLE LAWS AND
-// INTERNATIONAL TREATIES. THE RECEIPT OR POSSESSION OF THIS SOURCE
-// CODE AND/OR RELATED INFORMATION DOES NOT CONVEY OR IMPLY ANY RIGHTS
-// TO REPRODUCE, DISCLOSE OR DISTRIBUTE ITS CONTENTS, OR TO MANUFACTURE,
-// USE, OR SELL ANYTHING THAT IT MAY DESCRIBE, IN WHOLE OR IN PART.
-//
-// %COPYRIGHT_END%
-// --------------------------------------------------------------------
-// %BANNER_END%
+// Copyright 1998-2018 Epic Games, Inc. All Rights Reserved.
 
 #include "MagicLeapIdentity.h"
 #include "IMagicLeapIdentityPlugin.h"
diff --git a/Engine/Plugins/Lumin/MagicLeap/Source/MagicLeapIdentity/Public/IMagicLeapIdentityPlugin.h b/Engine/Plugins/Lumin/MagicLeap/Source/MagicLeapIdentity/Public/IMagicLeapIdentityPlugin.h
--- a/Engine/Plugins/Lumin/MagicLeap/Source/MagicLeapIdentity/Public/IMagicLeapIdentityPlugin.h
+++ b/Engine/Plugins/Lumin/MagicLeap/Source/MagicLeapIdentity/Public/IMagicLeapIdentityPlugin.h
@@ -1,36 +1,4 @@
-// %BANNER_BEGIN%
-// ---------------------------------------------------------------------
-// %COPYRIGHT_BEGIN%
-//
-// Copyright (c) 2017 Magic Leap, Inc. (COMPANY) All Rights Reserved.
-// Magic Leap, Inc. Confidential and Proprietary
-//
-// NOTICE: All information contained herein is, and remains the property
-// of COMPANY. The intellectual and technical concepts contained herein
-// are proprietary to COMPANY and may be covered by U.S. and Foreign
-// Patents, patents in process, and are protected by trade secret or
-// copyright law. Dissemination of this information or reproduction of
-// this material is strictly forbidden unless prior written permission is
-// obtained from COMPANY. Access to the source code contained herein is
-// hereby forbidden to anyone except current COMPANY employees, managers
-// or contractors who have executed Confidentiality and Non-disclosure
-// agreements explicitly covering such access.
-//
-// The copyright notice above does not evidence any actual or intended
-// publication or disclosure of this source code, which includes
-// information that is confidential and/or proprietary, and is a trade
-// secret, of COMPANY. ANY REPRODUCTION, MODIFICATION, DISTRIBUTION,
-// PUBLIC PERFORMANCE, OR PUBLIC DISPLAY OF OR THROUGH USE OF THIS
-// SOURCE CODE WITHOUT THE EXPRESS WRITTEN CONSENT OF COMPANY IS
-// STRICTLY PROHIBITED, AND IN VIOLATION OF APPLICABLE LAWS AND
-// INTERNATIONAL TREATIES. THE RECEIPT OR POSSESSION OF THIS SOURCE
-// CODE AND/OR RELATED INFORMATION DOES NOT CONVEY OR IMPLY ANY RIGHTS
-// TO REPRODUCE, DISCLOSE OR DISTRIBUTE ITS CONTENTS, OR TO MANUFACTURE,
-// USE, OR SELL ANYTHING THAT IT MAY DESCRIBE, IN WHOLE OR IN PART.
-//
-// %COPYRIGHT_END%
-// --------------------------------------------------------------------
-// %BANNER_END%
+// Copyright 1998-2018 Epic Games, Inc. All Rights Reserved.
 
 #pragma once
 
diff --git a/Engine/Plugins/Lumin/MagicLeap/Source/MagicLeapIdentity/Public/MagicLeapIdentity.h b/Engine/Plugins/Lumin/MagicLeap/Source/MagicLeapIdentity/Public/MagicLeapIdentity.h
--- a/Engine/Plugins/Lumin/MagicLeap/Source/MagicLeapIdentity/Public/MagicLeapIdentity.h
+++ b/Engine/Plugins/Lumin/MagicLeap/Source/MagicLeapIdentity/Public/MagicLeapIdentity.h
@@ -1,36 +1,4 @@
-// %BANNER_BEGIN%
-// ---------------------------------------------------------------------
-// %COPYRIGHT_BEGIN%
-//
-// Copyright (c) 2017 Magic Leap, Inc. (COMPANY) All Rights Reserved.
-// Magic Leap, Inc. Confidential and Proprietary
-//
-// NOTICE: All information contained herein is, and remains the property
-// of COMPANY. The intellectual and technical concepts contained herein
-// are proprietary to COMPANY and may be covered by U.S. and Foreign
-// Patents, patents in process, and are protected by trade secret or
-// copyright law. Dissemination of this information or reproduction of
-// this material is strictly forbidden unless prior written permission is
-// obtained from COMPANY. Access to the source code contained herein is
-// hereby forbidden to anyone except current COMPANY employees, managers
-// or contractors who have executed Confidentiality and Non-disclosure
-// agreements explicitly covering such access.
-//
-// The copyright notice above does not evidence any actual or intended
-// publication or disclosure of this source code, which includes
-// information that is confidential and/or proprietary, and is a trade
-// secret, of COMPANY. ANY REPRODUCTION, MODIFICATION, DISTRIBUTION,
-// PUBLIC PERFORMANCE, OR PUBLIC DISPLAY OF OR THROUGH USE OF THIS
-// SOURCE CODE WITHOUT THE EXPRESS WRITTEN CONSENT OF COMPANY IS
-// STRICTLY PROHIBITED, AND IN VIOLATION OF APPLICABLE LAWS AND
-// INTERNATIONAL TREATIES. THE RECEIPT OR POSSESSION OF THIS SOURCE
-// CODE AND/OR RELATED INFORMATION DOES NOT CONVEY OR IMPLY ANY RIGHTS
-// TO REPRODUCE, DISCLOSE OR DISTRIBUTE ITS CONTENTS, OR TO MANUFACTURE,
-// USE, OR SELL ANYTHING THAT IT MAY DESCRIBE, IN WHOLE OR IN PART.
-//
-// %COPYRIGHT_END%
-// --------------------------------------------------------------------
-// %BANNER_END%
+// Copyright 1998-2018 Epic Games, Inc. All Rights Reserved.
 
 #pragma once
 
diff --git a/Engine/Plugins/Lumin/MagicLeap/Source/MagicLeapIdentity/Public/MagicLeapIdentityTypes.h b/Engine/Plugins/Lumin/MagicLeap/Source/MagicLeapIdentity/Public/MagicLeapIdentityTypes.h
--- a/Engine/Plugins/Lumin/MagicLeap/Source/MagicLeapIdentity/Public/MagicLeapIdentityTypes.h
+++ b/Engine/Plugins/Lumin/MagicLeap/Source/MagicLeapIdentity/Public/MagicLeapIdentityTypes.h
@@ -1,36 +1,4 @@
-// %BANNER_BEGIN%
-// ---------------------------------------------------------------------
-// %COPYRIGHT_BEGIN%
-//
-// Copyright (c) 2017 Magic Leap, Inc. (COMPANY) All Rights Reserved.
-// Magic Leap, Inc. Confidential and Proprietary
-//
-// NOTICE: All information contained herein is, and remains the property
-// of COMPANY. The intellectual and technical concepts contained herein
-// are proprietary to COMPANY and may be covered by U.S. and Foreign
-// Patents, patents in process, and are protected by trade secret or
-// copyright law. Dissemination of this information or reproduction of
-// this material is strictly forbidden unless prior written permission is
-// obtained from COMPANY. Access to the source code contained herein is
-// hereby forbidden to anyone except current COMPANY employees, managers
-// or contractors who have executed Confidentiality and Non-disclosure
-// agreements explicitly covering such access.
-//
-// The copyright notice above does not evidence any actual or intended
-// publication or disclosure of this source code, which includes
-// information that is confidential and/or proprietary, and is a trade
-// secret, of COMPANY. ANY REPRODUCTION, MODIFICATION, DISTRIBUTION,
-// PUBLIC PERFORMANCE, OR PUBLIC DISPLAY OF OR THROUGH USE OF THIS
-// SOURCE CODE WITHOUT THE EXPRESS WRITTEN CONSENT OF COMPANY IS
-// STRICTLY PROHIBITED, AND IN VIOLATION OF APPLICABLE LAWS AND
-// INTERNATIONAL TREATIES. THE RECEIPT OR POSSESSION OF THIS SOURCE
-// CODE AND/OR RELATED INFORMATION DOES NOT CONVEY OR IMPLY ANY RIGHTS
-// TO REPRODUCE, DISCLOSE OR DISTRIBUTE ITS CONTENTS, OR TO MANUFACTURE,
-// USE, OR SELL ANYTHING THAT IT MAY DESCRIBE, IN WHOLE OR IN PART.
-//
-// %COPYRIGHT_END%
-// --------------------------------------------------------------------
-// %BANNER_END%
+// Copyright 1998-2018 Epic Games, Inc. All Rights Reserved.
 
 #pragma once
 
diff --git a/Engine/Plugins/Lumin/MagicLeap/Source/MagicLeapPrivileges/MagicLeapPrivileges.Build.cs b/Engine/Plugins/Lumin/MagicLeap/Source/MagicLeapPrivileges/MagicLeapPrivileges.Build.cs
--- a/Engine/Plugins/Lumin/MagicLeap/Source/MagicLeapPrivileges/MagicLeapPrivileges.Build.cs
+++ b/Engine/Plugins/Lumin/MagicLeap/Source/MagicLeapPrivileges/MagicLeapPrivileges.Build.cs
@@ -1,36 +1,4 @@
-// %BANNER_BEGIN%
-// ---------------------------------------------------------------------
-// %COPYRIGHT_BEGIN%
-//
-// Copyright (c) 2017 Magic Leap, Inc. (COMPANY) All Rights Reserved.
-// Magic Leap, Inc. Confidential and Proprietary
-//
-// NOTICE: All information contained herein is, and remains the property
-// of COMPANY. The intellectual and technical concepts contained herein
-// are proprietary to COMPANY and may be covered by U.S. and Foreign
-// Patents, patents in process, and are protected by trade secret or
-// copyright law. Dissemination of this information or reproduction of
-// this material is strictly forbidden unless prior written permission is
-// obtained from COMPANY. Access to the source code contained herein is
-// hereby forbidden to anyone except current COMPANY employees, managers
-// or contractors who have executed Confidentiality and Non-disclosure
-// agreements explicitly covering such access.
-//
-// The copyright notice above does not evidence any actual or intended
-// publication or disclosure of this source code, which includes
-// information that is confidential and/or proprietary, and is a trade
-// secret, of COMPANY. ANY REPRODUCTION, MODIFICATION, DISTRIBUTION,
-// PUBLIC PERFORMANCE, OR PUBLIC DISPLAY OF OR THROUGH USE OF THIS
-// SOURCE CODE WITHOUT THE EXPRESS WRITTEN CONSENT OF COMPANY IS
-// STRICTLY PROHIBITED, AND IN VIOLATION OF APPLICABLE LAWS AND
-// INTERNATIONAL TREATIES. THE RECEIPT OR POSSESSION OF THIS SOURCE
-// CODE AND/OR RELATED INFORMATION DOES NOT CONVEY OR IMPLY ANY RIGHTS
-// TO REPRODUCE, DISCLOSE OR DISTRIBUTE ITS CONTENTS, OR TO MANUFACTURE,
-// USE, OR SELL ANYTHING THAT IT MAY DESCRIBE, IN WHOLE OR IN PART.
-//
-// %COPYRIGHT_END%
-// --------------------------------------------------------------------
-// %BANNER_END%
+// Copyright 1998-2018 Epic Games, Inc. All Rights Reserved.
 
 using System.IO;
 
diff --git a/Engine/Plugins/Lumin/MagicLeap/Source/MagicLeapPrivileges/Private/MagicLeapPrivileges.cpp b/Engine/Plugins/Lumin/MagicLeap/Source/MagicLeapPrivileges/Private/MagicLeapPrivileges.cpp
--- a/Engine/Plugins/Lumin/MagicLeap/Source/MagicLeapPrivileges/Private/MagicLeapPrivileges.cpp
+++ b/Engine/Plugins/Lumin/MagicLeap/Source/MagicLeapPrivileges/Private/MagicLeapPrivileges.cpp
@@ -1,36 +1,4 @@
-// %BANNER_BEGIN%
-// ---------------------------------------------------------------------
-// %COPYRIGHT_BEGIN%
-//
-// Copyright (c) 2017 Magic Leap, Inc. (COMPANY) All Rights Reserved.
-// Magic Leap, Inc. Confidential and Proprietary
-//
-// NOTICE: All information contained herein is, and remains the property
-// of COMPANY. The intellectual and technical concepts contained herein
-// are proprietary to COMPANY and may be covered by U.S. and Foreign
-// Patents, patents in process, and are protected by trade secret or
-// copyright law. Dissemination of this information or reproduction of
-// this material is strictly forbidden unless prior written permission is
-// obtained from COMPANY. Access to the source code contained herein is
-// hereby forbidden to anyone except current COMPANY employees, managers
-// or contractors who have executed Confidentiality and Non-disclosure
-// agreements explicitly covering such access.
-//
-// The copyright notice above does not evidence any actual or intended
-// publication or disclosure of this source code, which includes
-// information that is confidential and/or proprietary, and is a trade
-// secret, of COMPANY. ANY REPRODUCTION, MODIFICATION, DISTRIBUTION,
-// PUBLIC PERFORMANCE, OR PUBLIC DISPLAY OF OR THROUGH USE OF THIS
-// SOURCE CODE WITHOUT THE EXPRESS WRITTEN CONSENT OF COMPANY IS
-// STRICTLY PROHIBITED, AND IN VIOLATION OF APPLICABLE LAWS AND
-// INTERNATIONAL TREATIES. THE RECEIPT OR POSSESSION OF THIS SOURCE
-// CODE AND/OR RELATED INFORMATION DOES NOT CONVEY OR IMPLY ANY RIGHTS
-// TO REPRODUCE, DISCLOSE OR DISTRIBUTE ITS CONTENTS, OR TO MANUFACTURE,
-// USE, OR SELL ANYTHING THAT IT MAY DESCRIBE, IN WHOLE OR IN PART.
-//
-// %COPYRIGHT_END%
-// --------------------------------------------------------------------
-// %BANNER_END%
+// Copyright 1998-2018 Epic Games, Inc. All Rights Reserved.
 
 #include "MagicLeapPrivileges.h"
 #include "IMagicLeapPrivilegesPlugin.h"
diff --git a/Engine/Plugins/Lumin/MagicLeap/Source/MagicLeapPrivileges/Public/IMagicLeapPrivilegesPlugin.h b/Engine/Plugins/Lumin/MagicLeap/Source/MagicLeapPrivileges/Public/IMagicLeapPrivilegesPlugin.h
--- a/Engine/Plugins/Lumin/MagicLeap/Source/MagicLeapPrivileges/Public/IMagicLeapPrivilegesPlugin.h
+++ b/Engine/Plugins/Lumin/MagicLeap/Source/MagicLeapPrivileges/Public/IMagicLeapPrivilegesPlugin.h
@@ -1,36 +1,4 @@
-// %BANNER_BEGIN%
-// ---------------------------------------------------------------------
-// %COPYRIGHT_BEGIN%
-//
-// Copyright (c) 2017 Magic Leap, Inc. (COMPANY) All Rights Reserved.
-// Magic Leap, Inc. Confidential and Proprietary
-//
-// NOTICE: All information contained herein is, and remains the property
-// of COMPANY. The intellectual and technical concepts contained herein
-// are proprietary to COMPANY and may be covered by U.S. and Foreign
-// Patents, patents in process, and are protected by trade secret or
-// copyright law. Dissemination of this information or reproduction of
-// this material is strictly forbidden unless prior written permission is
-// obtained from COMPANY. Access to the source code contained herein is
-// hereby forbidden to anyone except current COMPANY employees, managersMag
-// or contractors who have executed Confidentiality and Non-disclosure
-// agreements explicitly covering such access.
-//
-// The copyright notice above does not evidence any actual or intended
-// publication or disclosure of this source code, which includes
-// information that is confidential and/or proprietary, and is a trade
-// secret, of COMPANY. ANY REPRODUCTION, MODIFICATION, DISTRIBUTION,
-// PUBLIC PERFORMANCE, OR PUBLIC DISPLAY OF OR THROUGH USE OF THIS
-// SOURCE CODE WITHOUT THE EXPRESS WRITTEN CONSENT OF COMPANY IS
-// STRICTLY PROHIBITED, AND IN VIOLATION OF APPLICABLE LAWS AND
-// INTERNATIONAL TREATIES. THE RECEIPT OR POSSESSION OF THIS SOURCE
-// CODE AND/OR RELATED INFORMATION DOES NOT CONVEY OR IMPLY ANY RIGHTS
-// TO REPRODUCE, DISCLOSE OR DISTRIBUTE ITS CONTENTS, OR TO MANUFACTURE,
-// USE, OR SELL ANYTHING THAT IT MAY DESCRIBE, IN WHOLE OR IN PART.
-//
-// %COPYRIGHT_END%
-// --------------------------------------------------------------------
-// %BANNER_END%
+// Copyright 1998-2018 Epic Games, Inc. All Rights Reserved.
 
 #pragma once
 
diff --git a/Engine/Plugins/Lumin/MagicLeap/Source/MagicLeapPrivileges/Public/MagicLeapPrivileges.h b/Engine/Plugins/Lumin/MagicLeap/Source/MagicLeapPrivileges/Public/MagicLeapPrivileges.h
--- a/Engine/Plugins/Lumin/MagicLeap/Source/MagicLeapPrivileges/Public/MagicLeapPrivileges.h
+++ b/Engine/Plugins/Lumin/MagicLeap/Source/MagicLeapPrivileges/Public/MagicLeapPrivileges.h
@@ -1,36 +1,4 @@
-// %BANNER_BEGIN%
-// ---------------------------------------------------------------------
-// %COPYRIGHT_BEGIN%
-//
-// Copyright (c) 2017 Magic Leap, Inc. (COMPANY) All Rights Reserved.
-// Magic Leap, Inc. Confidential and Proprietary
-//
-// NOTICE: All information contained herein is, and remains the property
-// of COMPANY. The intellectual and technical concepts contained herein
-// are proprietary to COMPANY and may be covered by U.S. and Foreign
-// Patents, patents in process, and are protected by trade secret or
-// copyright law. Dissemination of this information or reproduction of
-// this material is strictly forbidden unless prior written permission is
-// obtained from COMPANY. Access to the source code contained herein is
-// hereby forbidden to anyone except current COMPANY employees, managers
-// or contractors who have executed Confidentiality and Non-disclosure
-// agreements explicitly covering such access.
-//
-// The copyright notice above does not evidence any actual or intended
-// publication or disclosure of this source code, which includes
-// information that is confidential and/or proprietary, and is a trade
-// secret, of COMPANY. ANY REPRODUCTION, MODIFICATION, DISTRIBUTION,
-// PUBLIC PERFORMANCE, OR PUBLIC DISPLAY OF OR THROUGH USE OF THIS
-// SOURCE CODE WITHOUT THE EXPRESS WRITTEN CONSENT OF COMPANY IS
-// STRICTLY PROHIBITED, AND IN VIOLATION OF APPLICABLE LAWS AND
-// INTERNATIONAL TREATIES. THE RECEIPT OR POSSESSION OF THIS SOURCE
-// CODE AND/OR RELATED INFORMATION DOES NOT CONVEY OR IMPLY ANY RIGHTS
-// TO REPRODUCE, DISCLOSE OR DISTRIBUTE ITS CONTENTS, OR TO MANUFACTURE,
-// USE, OR SELL ANYTHING THAT IT MAY DESCRIBE, IN WHOLE OR IN PART.
-//
-// %COPYRIGHT_END%
-// --------------------------------------------------------------------
-// %BANNER_END%
+// Copyright 1998-2018 Epic Games, Inc. All Rights Reserved.
 
 #pragma once
 
diff --git a/Engine/Plugins/Lumin/MagicLeap/Source/MagicLeapSecureStorage/MagicLeapSecureStorage.Build.cs b/Engine/Plugins/Lumin/MagicLeap/Source/MagicLeapSecureStorage/MagicLeapSecureStorage.Build.cs
--- a/Engine/Plugins/Lumin/MagicLeap/Source/MagicLeapSecureStorage/MagicLeapSecureStorage.Build.cs
+++ b/Engine/Plugins/Lumin/MagicLeap/Source/MagicLeapSecureStorage/MagicLeapSecureStorage.Build.cs
@@ -1,36 +1,4 @@
-// %BANNER_BEGIN%
-// ---------------------------------------------------------------------
-// %COPYRIGHT_BEGIN%
-//
-// Copyright (c) 2017 Magic Leap, Inc. (COMPANY) All Rights Reserved.
-// Magic Leap, Inc. Confidential and Proprietary
-//
-// NOTICE: All information contained herein is, and remains the property
-// of COMPANY. The intellectual and technical concepts contained herein
-// are proprietary to COMPANY and may be covered by U.S. and Foreign
-// Patents, patents in process, and are protected by trade secret or
-// copyright law. Dissemination of this information or reproduction of
-// this material is strictly forbidden unless prior written permission is
-// obtained from COMPANY. Access to the source code contained herein is
-// hereby forbidden to anyone except current COMPANY employees, managers
-// or contractors who have executed Confidentiality and Non-disclosure
-// agreements explicitly covering such access.
-//
-// The copyright notice above does not evidence any actual or intended
-// publication or disclosure of this source code, which includes
-// information that is confidential and/or proprietary, and is a trade
-// secret, of COMPANY. ANY REPRODUCTION, MODIFICATION, DISTRIBUTION,
-// PUBLIC PERFORMANCE, OR PUBLIC DISPLAY OF OR THROUGH USE OF THIS
-// SOURCE CODE WITHOUT THE EXPRESS WRITTEN CONSENT OF COMPANY IS
-// STRICTLY PROHIBITED, AND IN VIOLATION OF APPLICABLE LAWS AND
-// INTERNATIONAL TREATIES. THE RECEIPT OR POSSESSION OF THIS SOURCE
-// CODE AND/OR RELATED INFORMATION DOES NOT CONVEY OR IMPLY ANY RIGHTS
-// TO REPRODUCE, DISCLOSE OR DISTRIBUTE ITS CONTENTS, OR TO MANUFACTURE,
-// USE, OR SELL ANYTHING THAT IT MAY DESCRIBE, IN WHOLE OR IN PART.
-//
-// %COPYRIGHT_END%
-// --------------------------------------------------------------------
-// %BANNER_END%
+// Copyright 1998-2018 Epic Games, Inc. All Rights Reserved.
 
 using System.IO;
 
diff --git a/Engine/Plugins/Lumin/MagicLeap/Source/MagicLeapSecureStorage/Private/MagicLeapSecureStorage.cpp b/Engine/Plugins/Lumin/MagicLeap/Source/MagicLeapSecureStorage/Private/MagicLeapSecureStorage.cpp
--- a/Engine/Plugins/Lumin/MagicLeap/Source/MagicLeapSecureStorage/Private/MagicLeapSecureStorage.cpp
+++ b/Engine/Plugins/Lumin/MagicLeap/Source/MagicLeapSecureStorage/Private/MagicLeapSecureStorage.cpp
@@ -1,36 +1,4 @@
-// %BANNER_BEGIN%
-// ---------------------------------------------------------------------
-// %COPYRIGHT_BEGIN%
-//
-// Copyright (c) 2017 Magic Leap, Inc. (COMPANY) All Rights Reserved.
-// Magic Leap, Inc. Confidential and Proprietary
-//
-// NOTICE: All information contained herein is, and remains the property
-// of COMPANY. The intellectual and technical concepts contained herein
-// are proprietary to COMPANY and may be covered by U.S. and Foreign
-// Patents, patents in process, and are protected by trade secret or
-// copyright law. Dissemination of this information or reproduction of
-// this material is strictly forbidden unless prior written permission is
-// obtained from COMPANY. Access to the source code contained herein is
-// hereby forbidden to anyone except current COMPANY employees, managers
-// or contractors who have executed Confidentiality and Non-disclosure
-// agreements explicitly covering such access.
-//
-// The copyright notice above does not evidence any actual or intended
-// publication or disclosure of this source code, which includes
-// information that is confidential and/or proprietary, and is a trade
-// secret, of COMPANY. ANY REPRODUCTION, MODIFICATION, DISTRIBUTION,
-// PUBLIC PERFORMANCE, OR PUBLIC DISPLAY OF OR THROUGH USE OF THIS
-// SOURCE CODE WITHOUT THE EXPRESS WRITTEN CONSENT OF COMPANY IS
-// STRICTLY PROHIBITED, AND IN VIOLATION OF APPLICABLE LAWS AND
-// INTERNATIONAL TREATIES. THE RECEIPT OR POSSESSION OF THIS SOURCE
-// CODE AND/OR RELATED INFORMATION DOES NOT CONVEY OR IMPLY ANY RIGHTS
-// TO REPRODUCE, DISCLOSE OR DISTRIBUTE ITS CONTENTS, OR TO MANUFACTURE,
-// USE, OR SELL ANYTHING THAT IT MAY DESCRIBE, IN WHOLE OR IN PART.
-//
-// %COPYRIGHT_END%
-// --------------------------------------------------------------------
-// %BANNER_END%
+// Copyright 1998-2018 Epic Games, Inc. All Rights Reserved.
 
 #include "MagicLeapSecureStorage.h"
 #include "IMagicLeapSecureStoragePlugin.h"
diff --git a/Engine/Plugins/Lumin/MagicLeap/Source/MagicLeapSecureStorage/Public/IMagicLeapSecureStoragePlugin.h b/Engine/Plugins/Lumin/MagicLeap/Source/MagicLeapSecureStorage/Public/IMagicLeapSecureStoragePlugin.h
--- a/Engine/Plugins/Lumin/MagicLeap/Source/MagicLeapSecureStorage/Public/IMagicLeapSecureStoragePlugin.h
+++ b/Engine/Plugins/Lumin/MagicLeap/Source/MagicLeapSecureStorage/Public/IMagicLeapSecureStoragePlugin.h
@@ -1,36 +1,4 @@
-// %BANNER_BEGIN%
-// ---------------------------------------------------------------------
-// %COPYRIGHT_BEGIN%
-//
-// Copyright (c) 2017 Magic Leap, Inc. (COMPANY) All Rights Reserved.
-// Magic Leap, Inc. Confidential and Proprietary
-//
-// NOTICE: All information contained herein is, and remains the property
-// of COMPANY. The intellectual and technical concepts contained herein
-// are proprietary to COMPANY and may be covered by U.S. and Foreign
-// Patents, patents in process, and are protected by trade secret or
-// copyright law. Dissemination of this information or reproduction of
-// this material is strictly forbidden unless prior written permission is
-// obtained from COMPANY. Access to the source code contained herein is
-// hereby forbidden to anyone except current COMPANY employees, managers
-// or contractors who have executed Confidentiality and Non-disclosure
-// agreements explicitly covering such access.
-//
-// The copyright notice above does not evidence any actual or intended
-// publication or disclosure of this source code, which includes
-// information that is confidential and/or proprietary, and is a trade
-// secret, of COMPANY. ANY REPRODUCTION, MODIFICATION, DISTRIBUTION,
-// PUBLIC PERFORMANCE, OR PUBLIC DISPLAY OF OR THROUGH USE OF THIS
-// SOURCE CODE WITHOUT THE EXPRESS WRITTEN CONSENT OF COMPANY IS
-// STRICTLY PROHIBITED, AND IN VIOLATION OF APPLICABLE LAWS AND
-// INTERNATIONAL TREATIES. THE RECEIPT OR POSSESSION OF THIS SOURCE
-// CODE AND/OR RELATED INFORMATION DOES NOT CONVEY OR IMPLY ANY RIGHTS
-// TO REPRODUCE, DISCLOSE OR DISTRIBUTE ITS CONTENTS, OR TO MANUFACTURE,
-// USE, OR SELL ANYTHING THAT IT MAY DESCRIBE, IN WHOLE OR IN PART.
-//
-// %COPYRIGHT_END%
-// --------------------------------------------------------------------
-// %BANNER_END%
+// Copyright 1998-2018 Epic Games, Inc. All Rights Reserved.
 
 #pragma once
 
diff --git a/Engine/Plugins/Lumin/MagicLeap/Source/MagicLeapSecureStorage/Public/MagicLeapSecureStorage.h b/Engine/Plugins/Lumin/MagicLeap/Source/MagicLeapSecureStorage/Public/MagicLeapSecureStorage.h
--- a/Engine/Plugins/Lumin/MagicLeap/Source/MagicLeapSecureStorage/Public/MagicLeapSecureStorage.h
+++ b/Engine/Plugins/Lumin/MagicLeap/Source/MagicLeapSecureStorage/Public/MagicLeapSecureStorage.h
@@ -1,36 +1,4 @@
-// %BANNER_BEGIN%
-// ---------------------------------------------------------------------
-// %COPYRIGHT_BEGIN%
-//
-// Copyright (c) 2017 Magic Leap, Inc. (COMPANY) All Rights Reserved.
-// Magic Leap, Inc. Confidential and Proprietary
-//
-// NOTICE: All information contained herein is, and remains the property
-// of COMPANY. The intellectual and technical concepts contained herein
-// are proprietary to COMPANY and may be covered by U.S. and Foreign
-// Patents, patents in process, and are protected by trade secret or
-// copyright law. Dissemination of this information or reproduction of
-// this material is strictly forbidden unless prior written permission is
-// obtained from COMPANY. Access to the source code contained herein is
-// hereby forbidden to anyone except current COMPANY employees, managers
-// or contractors who have executed Confidentiality and Non-disclosure
-// agreements explicitly covering such access.
-//
-// The copyright notice above does not evidence any actual or intended
-// publication or disclosure of this source code, which includes
-// information that is confidential and/or proprietary, and is a trade
-// secret, of COMPANY. ANY REPRODUCTION, MODIFICATION, DISTRIBUTION,
-// PUBLIC PERFORMANCE, OR PUBLIC DISPLAY OF OR THROUGH USE OF THIS
-// SOURCE CODE WITHOUT THE EXPRESS WRITTEN CONSENT OF COMPANY IS
-// STRICTLY PROHIBITED, AND IN VIOLATION OF APPLICABLE LAWS AND
-// INTERNATIONAL TREATIES. THE RECEIPT OR POSSESSION OF THIS SOURCE
-// CODE AND/OR RELATED INFORMATION DOES NOT CONVEY OR IMPLY ANY RIGHTS
-// TO REPRODUCE, DISCLOSE OR DISTRIBUTE ITS CONTENTS, OR TO MANUFACTURE,
-// USE, OR SELL ANYTHING THAT IT MAY DESCRIBE, IN WHOLE OR IN PART.
-//
-// %COPYRIGHT_END%
-// --------------------------------------------------------------------
-// %BANNER_END%
+// Copyright 1998-2018 Epic Games, Inc. All Rights Reserved.
 
 #pragma once
 
diff --git a/Engine/Plugins/Lumin/MagicLeapAnalytics/Source/MagicLeapAnalytics/MagicLeapAnalytics.Build.cs b/Engine/Plugins/Lumin/MagicLeapAnalytics/Source/MagicLeapAnalytics/MagicLeapAnalytics.Build.cs
--- a/Engine/Plugins/Lumin/MagicLeapAnalytics/Source/MagicLeapAnalytics/MagicLeapAnalytics.Build.cs
+++ b/Engine/Plugins/Lumin/MagicLeapAnalytics/Source/MagicLeapAnalytics/MagicLeapAnalytics.Build.cs
@@ -1,36 +1,4 @@
-// %BANNER_BEGIN%
-// ---------------------------------------------------------------------
-// %COPYRIGHT_BEGIN%
-//
-// Copyright (c) 2017 Magic Leap, Inc. (COMPANY) All Rights Reserved.
-// Magic Leap, Inc. Confidential and Proprietary
-//
-// NOTICE: All information contained herein is, and remains the property
-// of COMPANY. The intellectual and technical concepts contained herein
-// are proprietary to COMPANY and may be covered by U.S. and Foreign
-// Patents, patents in process, and are protected by trade secret or
-// copyright law. Dissemination of this information or reproduction of
-// this material is strictly forbidden unless prior written permission is
-// obtained from COMPANY. Access to the source code contained herein is
-// hereby forbidden to anyone except current COMPANY employees, managers
-// or contractors who have executed Confidentiality and Non-disclosure
-// agreements explicitly covering such access.
-//
-// The copyright notice above does not evidence any actual or intended
-// publication or disclosure of this source code, which includes
-// information that is confidential and/or proprietary, and is a trade
-// secret, of COMPANY. ANY REPRODUCTION, MODIFICATION, DISTRIBUTION,
-// PUBLIC PERFORMANCE, OR PUBLIC DISPLAY OF OR THROUGH USE OF THIS
-// SOURCE CODE WITHOUT THE EXPRESS WRITTEN CONSENT OF COMPANY IS
-// STRICTLY PROHIBITED, AND IN VIOLATION OF APPLICABLE LAWS AND
-// INTERNATIONAL TREATIES. THE RECEIPT OR POSSESSION OF THIS SOURCE
-// CODE AND/OR RELATED INFORMATION DOES NOT CONVEY OR IMPLY ANY RIGHTS
-// TO REPRODUCE, DISCLOSE OR DISTRIBUTE ITS CONTENTS, OR TO MANUFACTURE,
-// USE, OR SELL ANYTHING THAT IT MAY DESCRIBE, IN WHOLE OR IN PART.
-//
-// %COPYRIGHT_END%
-// --------------------------------------------------------------------
-// %BANNER_END%
+// Copyright 1998-2018 Epic Games, Inc. All Rights Reserved.
 
 using System.IO;
 
diff --git a/Engine/Plugins/Lumin/MagicLeapAnalytics/Source/MagicLeapAnalytics/Private/MagicLeapAnalytics.cpp b/Engine/Plugins/Lumin/MagicLeapAnalytics/Source/MagicLeapAnalytics/Private/MagicLeapAnalytics.cpp
--- a/Engine/Plugins/Lumin/MagicLeapAnalytics/Source/MagicLeapAnalytics/Private/MagicLeapAnalytics.cpp
+++ b/Engine/Plugins/Lumin/MagicLeapAnalytics/Source/MagicLeapAnalytics/Private/MagicLeapAnalytics.cpp
@@ -1,36 +1,4 @@
-// %BANNER_BEGIN%
-// ---------------------------------------------------------------------
-// %COPYRIGHT_BEGIN%
-//
-// Copyright (c) 2017 Magic Leap, Inc. (COMPANY) All Rights Reserved.
-// Magic Leap, Inc. Confidential and Proprietary
-//
-// NOTICE: All information contained herein is, and remains the property
-// of COMPANY. The intellectual and technical concepts contained herein
-// are proprietary to COMPANY and may be covered by U.S. and Foreign
-// Patents, patents in process, and are protected by trade secret or
-// copyright law. Dissemination of this information or reproduction of
-// this material is strictly forbidden unless prior written permission is
-// obtained from COMPANY. Access to the source code contained herein is
-// hereby forbidden to anyone except current COMPANY employees, managers
-// or contractors who have executed Confidentiality and Non-disclosure
-// agreements explicitly covering such access.
-//
-// The copyright notice above does not evidence any actual or intended
-// publication or disclosure of this source code, which includes
-// information that is confidential and/or proprietary, and is a trade
-// secret, of COMPANY. ANY REPRODUCTION, MODIFICATION, DISTRIBUTION,
-// PUBLIC PERFORMANCE, OR PUBLIC DISPLAY OF OR THROUGH USE OF THIS
-// SOURCE CODE WITHOUT THE EXPRESS WRITTEN CONSENT OF COMPANY IS
-// STRICTLY PROHIBITED, AND IN VIOLATION OF APPLICABLE LAWS AND
-// INTERNATIONAL TREATIES. THE RECEIPT OR POSSESSION OF THIS SOURCE
-// CODE AND/OR RELATED INFORMATION DOES NOT CONVEY OR IMPLY ANY RIGHTS
-// TO REPRODUCE, DISCLOSE OR DISTRIBUTE ITS CONTENTS, OR TO MANUFACTURE,
-// USE, OR SELL ANYTHING THAT IT MAY DESCRIBE, IN WHOLE OR IN PART.
-//
-// %COPYRIGHT_END%
-// --------------------------------------------------------------------
-// %BANNER_END%
+// Copyright 1998-2018 Epic Games, Inc. All Rights Reserved.
 
 #include "MagicLeapAnalytics.h"
 #include "AnalyticsEventAttribute.h"
diff --git a/Engine/Plugins/Lumin/MagicLeapAnalytics/Source/MagicLeapAnalytics/Private/MagicLeapAnalyticsProvider.h b/Engine/Plugins/Lumin/MagicLeapAnalytics/Source/MagicLeapAnalytics/Private/MagicLeapAnalyticsProvider.h
--- a/Engine/Plugins/Lumin/MagicLeapAnalytics/Source/MagicLeapAnalytics/Private/MagicLeapAnalyticsProvider.h
+++ b/Engine/Plugins/Lumin/MagicLeapAnalytics/Source/MagicLeapAnalytics/Private/MagicLeapAnalyticsProvider.h
@@ -1,36 +1,4 @@
-// %BANNER_BEGIN%
-// ---------------------------------------------------------------------
-// %COPYRIGHT_BEGIN%
-//
-// Copyright (c) 2017 Magic Leap, Inc. (COMPANY) All Rights Reserved.
-// Magic Leap, Inc. Confidential and Proprietary
-//
-// NOTICE: All information contained herein is, and remains the property
-// of COMPANY. The intellectual and technical concepts contained herein
-// are proprietary to COMPANY and may be covered by U.S. and Foreign
-// Patents, patents in process, and are protected by trade secret or
-// copyright law. Dissemination of this information or reproduction of
-// this material is strictly forbidden unless prior written permission is
-// obtained from COMPANY. Access to the source code contained herein is
-// hereby forbidden to anyone except current COMPANY employees, managers
-// or contractors who have executed Confidentiality and Non-disclosure
-// agreements explicitly covering such access.
-//
-// The copyright notice above does not evidence any actual or intended
-// publication or disclosure of this source code, which includes
-// information that is confidential and/or proprietary, and is a trade
-// secret, of COMPANY. ANY REPRODUCTION, MODIFICATION, DISTRIBUTION,
-// PUBLIC PERFORMANCE, OR PUBLIC DISPLAY OF OR THROUGH USE OF THIS
-// SOURCE CODE WITHOUT THE EXPRESS WRITTEN CONSENT OF COMPANY IS
-// STRICTLY PROHIBITED, AND IN VIOLATION OF APPLICABLE LAWS AND
-// INTERNATIONAL TREATIES. THE RECEIPT OR POSSESSION OF THIS SOURCE
-// CODE AND/OR RELATED INFORMATION DOES NOT CONVEY OR IMPLY ANY RIGHTS
-// TO REPRODUCE, DISCLOSE OR DISTRIBUTE ITS CONTENTS, OR TO MANUFACTURE,
-// USE, OR SELL ANYTHING THAT IT MAY DESCRIBE, IN WHOLE OR IN PART.
-//
-// %COPYRIGHT_END%
-// --------------------------------------------------------------------
-// %BANNER_END%
+// Copyright 1998-2018 Epic Games, Inc. All Rights Reserved.
 
 #pragma once
 
diff --git a/Engine/Plugins/Lumin/MagicLeapAnalytics/Source/MagicLeapAnalytics/Public/MagicLeapAnalytics.h b/Engine/Plugins/Lumin/MagicLeapAnalytics/Source/MagicLeapAnalytics/Public/MagicLeapAnalytics.h
--- a/Engine/Plugins/Lumin/MagicLeapAnalytics/Source/MagicLeapAnalytics/Public/MagicLeapAnalytics.h
+++ b/Engine/Plugins/Lumin/MagicLeapAnalytics/Source/MagicLeapAnalytics/Public/MagicLeapAnalytics.h
@@ -1,36 +1,4 @@
-// %BANNER_BEGIN%
-// ---------------------------------------------------------------------
-// %COPYRIGHT_BEGIN%
-//
-// Copyright (c) 2017 Magic Leap, Inc. (COMPANY) All Rights Reserved.
-// Magic Leap, Inc. Confidential and Proprietary
-//
-// NOTICE: All information contained herein is, and remains the property
-// of COMPANY. The intellectual and technical concepts contained herein
-// are proprietary to COMPANY and may be covered by U.S. and Foreign
-// Patents, patents in process, and are protected by trade secret or
-// copyright law. Dissemination of this information or reproduction of
-// this material is strictly forbidden unless prior written permission is
-// obtained from COMPANY. Access to the source code contained herein is
-// hereby forbidden to anyone except current COMPANY employees, managers
-// or contractors who have executed Confidentiality and Non-disclosure
-// agreements explicitly covering such access.
-//
-// The copyright notice above does not evidence any actual or intended
-// publication or disclosure of this source code, which includes
-// information that is confidential and/or proprietary, and is a trade
-// secret, of COMPANY. ANY REPRODUCTION, MODIFICATION, DISTRIBUTION,
-// PUBLIC PERFORMANCE, OR PUBLIC DISPLAY OF OR THROUGH USE OF THIS
-// SOURCE CODE WITHOUT THE EXPRESS WRITTEN CONSENT OF COMPANY IS
-// STRICTLY PROHIBITED, AND IN VIOLATION OF APPLICABLE LAWS AND
-// INTERNATIONAL TREATIES. THE RECEIPT OR POSSESSION OF THIS SOURCE
-// CODE AND/OR RELATED INFORMATION DOES NOT CONVEY OR IMPLY ANY RIGHTS
-// TO REPRODUCE, DISCLOSE OR DISTRIBUTE ITS CONTENTS, OR TO MANUFACTURE,
-// USE, OR SELL ANYTHING THAT IT MAY DESCRIBE, IN WHOLE OR IN PART.
-//
-// %COPYRIGHT_END%
-// --------------------------------------------------------------------
-// %BANNER_END%
+// Copyright 1998-2018 Epic Games, Inc. All Rights Reserved.
 
 #pragma once
 
diff --git a/Engine/Plugins/Lumin/MagicLeapEmulator/Source/MagicLeapEmulator/Public/IMagicLeapEmulatorPlugin.h b/Engine/Plugins/Lumin/MagicLeapEmulator/Source/MagicLeapEmulator/Public/IMagicLeapEmulatorPlugin.h
--- a/Engine/Plugins/Lumin/MagicLeapEmulator/Source/MagicLeapEmulator/Public/IMagicLeapEmulatorPlugin.h
+++ b/Engine/Plugins/Lumin/MagicLeapEmulator/Source/MagicLeapEmulator/Public/IMagicLeapEmulatorPlugin.h
@@ -1,3 +1,4 @@
+// Copyright 1998-2018 Epic Games, Inc. All Rights Reserved.
 
 #pragma once
 
diff --git a/Engine/Plugins/Lumin/MagicLeapMedia/Source/MagicLeapCameraPreviewFactory/MagicLeapCameraPreviewFactory.Build.cs b/Engine/Plugins/Lumin/MagicLeapMedia/Source/MagicLeapCameraPreviewFactory/MagicLeapCameraPreviewFactory.Build.cs
--- a/Engine/Plugins/Lumin/MagicLeapMedia/Source/MagicLeapCameraPreviewFactory/MagicLeapCameraPreviewFactory.Build.cs
+++ b/Engine/Plugins/Lumin/MagicLeapMedia/Source/MagicLeapCameraPreviewFactory/MagicLeapCameraPreviewFactory.Build.cs
@@ -1,36 +1,4 @@
-// %BANNER_BEGIN%
-// ---------------------------------------------------------------------
-// %COPYRIGHT_BEGIN%
-//
-// Copyright (c) 2017 Magic Leap, Inc. (COMPANY) All Rights Reserved.
-// Magic Leap, Inc. Confidential and Proprietary
-//
-// NOTICE:  All information contained herein is, and remains the property
-// of COMPANY. The intellectual and technical concepts contained herein
-// are proprietary to COMPANY and may be covered by U.S. and Foreign
-// Patents, patents in process, and are protected by trade secret or
-// copyright law.  Dissemination of this information or reproduction of
-// this material is strictly forbidden unless prior written permission is
-// obtained from COMPANY.  Access to the source code contained herein is
-// hereby forbidden to anyone except current COMPANY employees, managers
-// or contractors who have executed Confidentiality and Non-disclosure
-// agreements explicitly covering such access.
-//
-// The copyright notice above does not evidence any actual or intended
-// publication or disclosure  of  this source code, which includes
-// information that is confidential and/or proprietary, and is a trade
-// secret, of  COMPANY.   ANY REPRODUCTION, MODIFICATION, DISTRIBUTION,
-// PUBLIC  PERFORMANCE, OR PUBLIC DISPLAY OF OR THROUGH USE  OF THIS
-// SOURCE CODE  WITHOUT THE EXPRESS WRITTEN CONSENT OF COMPANY IS
-// STRICTLY PROHIBITED, AND IN VIOLATION OF APPLICABLE LAWS AND
-// INTERNATIONAL TREATIES.  THE RECEIPT OR POSSESSION OF  THIS SOURCE
-// CODE AND/OR RELATED INFORMATION DOES NOT CONVEY OR IMPLY ANY RIGHTS
-// TO REPRODUCE, DISCLOSE OR DISTRIBUTE ITS CONTENTS, OR TO MANUFACTURE,
-// USE, OR SELL ANYTHING THAT IT  MAY DESCRIBE, IN WHOLE OR IN PART.
-//
-// %COPYRIGHT_END%
-// --------------------------------------------------------------------*/
-// %BANNER_END%
+// Copyright 1998-2018 Epic Games, Inc. All Rights Reserved.
 
 namespace UnrealBuildTool.Rules
 {
diff --git a/Engine/Plugins/Lumin/MagicLeapMedia/Source/MagicLeapCameraPreviewFactory/Private/MagicLeapCameraPreviewFactoryModule.cpp b/Engine/Plugins/Lumin/MagicLeapMedia/Source/MagicLeapCameraPreviewFactory/Private/MagicLeapCameraPreviewFactoryModule.cpp
--- a/Engine/Plugins/Lumin/MagicLeapMedia/Source/MagicLeapCameraPreviewFactory/Private/MagicLeapCameraPreviewFactoryModule.cpp
+++ b/Engine/Plugins/Lumin/MagicLeapMedia/Source/MagicLeapCameraPreviewFactory/Private/MagicLeapCameraPreviewFactoryModule.cpp
@@ -1,36 +1,4 @@
-// %BANNER_BEGIN%
-// ---------------------------------------------------------------------
-// %COPYRIGHT_BEGIN%
-//
-// Copyright (c) 2017 Magic Leap, Inc. (COMPANY) All Rights Reserved.
-// Magic Leap, Inc. Confidential and Proprietary
-//
-// NOTICE:  All information contained herein is, and remains the property
-// of COMPANY. The intellectual and technical concepts contained herein
-// are proprietary to COMPANY and may be covered by U.S. and Foreign
-// Patents, patents in process, and are protected by trade secret or
-// copyright law.  Dissemination of this information or reproduction of
-// this material is strictly forbidden unless prior written permission is
-// obtained from COMPANY.  Access to the source code contained herein is
-// hereby forbidden to anyone except current COMPANY employees, managers
-// or contractors who have executed Confidentiality and Non-disclosure
-// agreements explicitly covering such access.
-//
-// The copyright notice above does not evidence any actual or intended
-// publication or disclosure  of  this source code, which includes
-// information that is confidential and/or proprietary, and is a trade
-// secret, of  COMPANY.   ANY REPRODUCTION, MODIFICATION, DISTRIBUTION,
-// PUBLIC  PERFORMANCE, OR PUBLIC DISPLAY OF OR THROUGH USE  OF THIS
-// SOURCE CODE  WITHOUT THE EXPRESS WRITTEN CONSENT OF COMPANY IS
-// STRICTLY PROHIBITED, AND IN VIOLATION OF APPLICABLE LAWS AND
-// INTERNATIONAL TREATIES.  THE RECEIPT OR POSSESSION OF  THIS SOURCE
-// CODE AND/OR RELATED INFORMATION DOES NOT CONVEY OR IMPLY ANY RIGHTS
-// TO REPRODUCE, DISCLOSE OR DISTRIBUTE ITS CONTENTS, OR TO MANUFACTURE,
-// USE, OR SELL ANYTHING THAT IT  MAY DESCRIBE, IN WHOLE OR IN PART.
-//
-// %COPYRIGHT_END%
-// --------------------------------------------------------------------*/
-// %BANNER_END%
+// Copyright 1998-2018 Epic Games, Inc. All Rights Reserved.
 
 #include "CoreMinimal.h"
 #include "Misc/Paths.h"
diff --git a/Engine/Plugins/Lumin/MagicLeapMedia/Source/MagicLeapMedia/MagicLeapMedia.Build.cs b/Engine/Plugins/Lumin/MagicLeapMedia/Source/MagicLeapMedia/MagicLeapMedia.Build.cs
--- a/Engine/Plugins/Lumin/MagicLeapMedia/Source/MagicLeapMedia/MagicLeapMedia.Build.cs
+++ b/Engine/Plugins/Lumin/MagicLeapMedia/Source/MagicLeapMedia/MagicLeapMedia.Build.cs
@@ -1,36 +1,4 @@
-// %BANNER_BEGIN%
-// ---------------------------------------------------------------------
-// %COPYRIGHT_BEGIN%
-//
-// Copyright (c) 2017 Magic Leap, Inc. (COMPANY) All Rights Reserved.
-// Magic Leap, Inc. Confidential and Proprietary
-//
-// NOTICE:  All information contained herein is, and remains the property
-// of COMPANY. The intellectual and technical concepts contained herein
-// are proprietary to COMPANY and may be covered by U.S. and Foreign
-// Patents, patents in process, and are protected by trade secret or
-// copyright law.  Dissemination of this information or reproduction of
-// this material is strictly forbidden unless prior written permission is
-// obtained from COMPANY.  Access to the source code contained herein is
-// hereby forbidden to anyone except current COMPANY employees, managers
-// or contractors who have executed Confidentiality and Non-disclosure
-// agreements explicitly covering such access.
-//
-// The copyright notice above does not evidence any actual or intended
-// publication or disclosure  of  this source code, which includes
-// information that is confidential and/or proprietary, and is a trade
-// secret, of  COMPANY.   ANY REPRODUCTION, MODIFICATION, DISTRIBUTION,
-// PUBLIC  PERFORMANCE, OR PUBLIC DISPLAY OF OR THROUGH USE  OF THIS
-// SOURCE CODE  WITHOUT THE EXPRESS WRITTEN CONSENT OF COMPANY IS
-// STRICTLY PROHIBITED, AND IN VIOLATION OF APPLICABLE LAWS AND
-// INTERNATIONAL TREATIES.  THE RECEIPT OR POSSESSION OF  THIS SOURCE
-// CODE AND/OR RELATED INFORMATION DOES NOT CONVEY OR IMPLY ANY RIGHTS
-// TO REPRODUCE, DISCLOSE OR DISTRIBUTE ITS CONTENTS, OR TO MANUFACTURE,
-// USE, OR SELL ANYTHING THAT IT  MAY DESCRIBE, IN WHOLE OR IN PART.
-//
-// %COPYRIGHT_END%
-// --------------------------------------------------------------------*/
-// %BANNER_END%
+// Copyright 1998-2018 Epic Games, Inc. All Rights Reserved.
 
 using System.IO;
 
diff --git a/Engine/Plugins/Lumin/MagicLeapMedia/Source/MagicLeapMedia/Private/ExternalOESTextureRenderer.cpp b/Engine/Plugins/Lumin/MagicLeapMedia/Source/MagicLeapMedia/Private/ExternalOESTextureRenderer.cpp
--- a/Engine/Plugins/Lumin/MagicLeapMedia/Source/MagicLeapMedia/Private/ExternalOESTextureRenderer.cpp
+++ b/Engine/Plugins/Lumin/MagicLeapMedia/Source/MagicLeapMedia/Private/ExternalOESTextureRenderer.cpp
@@ -1,36 +1,4 @@
-// %BANNER_BEGIN%
-// ---------------------------------------------------------------------
-// %COPYRIGHT_BEGIN%
-//
-// Copyright (c) 2017 Magic Leap, Inc. (COMPANY) All Rights Reserved.
-// Magic Leap, Inc. Confidential and Proprietary
-//
-// NOTICE:  All information contained herein is, and remains the property
-// of COMPANY. The intellectual and technical concepts contained herein
-// are proprietary to COMPANY and may be covered by U.S. and Foreign
-// Patents, patents in process, and are protected by trade secret or
-// copyright law.  Dissemination of this information or reproduction of
-// this material is strictly forbidden unless prior written permission is
-// obtained from COMPANY.  Access to the source code contained herein is
-// hereby forbidden to anyone except current COMPANY employees, managers
-// or contractors who have executed Confidentiality and Non-disclosure
-// agreements explicitly covering such access.
-//
-// The copyright notice above does not evidence any actual or intended
-// publication or disclosure  of  this source code, which includes
-// information that is confidential and/or proprietary, and is a trade
-// secret, of  COMPANY.   ANY REPRODUCTION, MODIFICATION, DISTRIBUTION,
-// PUBLIC  PERFORMANCE, OR PUBLIC DISPLAY OF OR THROUGH USE  OF THIS
-// SOURCE CODE  WITHOUT THE EXPRESS WRITTEN CONSENT OF COMPANY IS
-// STRICTLY PROHIBITED, AND IN VIOLATION OF APPLICABLE LAWS AND
-// INTERNATIONAL TREATIES.  THE RECEIPT OR POSSESSION OF  THIS SOURCE
-// CODE AND/OR RELATED INFORMATION DOES NOT CONVEY OR IMPLY ANY RIGHTS
-// TO REPRODUCE, DISCLOSE OR DISTRIBUTE ITS CONTENTS, OR TO MANUFACTURE,
-// USE, OR SELL ANYTHING THAT IT  MAY DESCRIBE, IN WHOLE OR IN PART.
-//
-// %COPYRIGHT_END%
-// --------------------------------------------------------------------*/
-// %BANNER_END%
+// Copyright 1998-2018 Epic Games, Inc. All Rights Reserved.
 
 #include "ExternalOESTextureRenderer.h"
 #include "RenderingThread.h"
diff --git a/Engine/Plugins/Lumin/MagicLeapMedia/Source/MagicLeapMedia/Private/ExternalOESTextureRenderer.h b/Engine/Plugins/Lumin/MagicLeapMedia/Source/MagicLeapMedia/Private/ExternalOESTextureRenderer.h
--- a/Engine/Plugins/Lumin/MagicLeapMedia/Source/MagicLeapMedia/Private/ExternalOESTextureRenderer.h
+++ b/Engine/Plugins/Lumin/MagicLeapMedia/Source/MagicLeapMedia/Private/ExternalOESTextureRenderer.h
@@ -1,36 +1,4 @@
-// %BANNER_BEGIN%
-// ---------------------------------------------------------------------
-// %COPYRIGHT_BEGIN%
-//
-// Copyright (c) 2017 Magic Leap, Inc. (COMPANY) All Rights Reserved.
-// Magic Leap, Inc. Confidential and Proprietary
-//
-// NOTICE:  All information contained herein is, and remains the property
-// of COMPANY. The intellectual and technical concepts contained herein
-// are proprietary to COMPANY and may be covered by U.S. and Foreign
-// Patents, patents in process, and are protected by trade secret or
-// copyright law.  Dissemination of this information or reproduction of
-// this material is strictly forbidden unless prior written permission is
-// obtained from COMPANY.  Access to the source code contained herein is
-// hereby forbidden to anyone except current COMPANY employees, managers
-// or contractors who have executed Confidentiality and Non-disclosure
-// agreements explicitly covering such access.
-//
-// The copyright notice above does not evidence any actual or intended
-// publication or disclosure  of  this source code, which includes
-// information that is confidential and/or proprietary, and is a trade
-// secret, of  COMPANY.   ANY REPRODUCTION, MODIFICATION, DISTRIBUTION,
-// PUBLIC  PERFORMANCE, OR PUBLIC DISPLAY OF OR THROUGH USE  OF THIS
-// SOURCE CODE  WITHOUT THE EXPRESS WRITTEN CONSENT OF COMPANY IS
-// STRICTLY PROHIBITED, AND IN VIOLATION OF APPLICABLE LAWS AND
-// INTERNATIONAL TREATIES.  THE RECEIPT OR POSSESSION OF  THIS SOURCE
-// CODE AND/OR RELATED INFORMATION DOES NOT CONVEY OR IMPLY ANY RIGHTS
-// TO REPRODUCE, DISCLOSE OR DISTRIBUTE ITS CONTENTS, OR TO MANUFACTURE,
-// USE, OR SELL ANYTHING THAT IT  MAY DESCRIBE, IN WHOLE OR IN PART.
-//
-// %COPYRIGHT_END%
-// --------------------------------------------------------------------*/
-// %BANNER_END%
+// Copyright 1998-2018 Epic Games, Inc. All Rights Reserved.
 
 #pragma once
 
diff --git a/Engine/Plugins/Lumin/MagicLeapMedia/Source/MagicLeapMedia/Private/MagicLeapCameraPreviewPlayer.cpp b/Engine/Plugins/Lumin/MagicLeapMedia/Source/MagicLeapMedia/Private/MagicLeapCameraPreviewPlayer.cpp
--- a/Engine/Plugins/Lumin/MagicLeapMedia/Source/MagicLeapMedia/Private/MagicLeapCameraPreviewPlayer.cpp
+++ b/Engine/Plugins/Lumin/MagicLeapMedia/Source/MagicLeapMedia/Private/MagicLeapCameraPreviewPlayer.cpp
@@ -1,36 +1,4 @@
-// %BANNER_BEGIN%
-// ---------------------------------------------------------------------
-// %COPYRIGHT_BEGIN%
-//
-// Copyright (c) 2017 Magic Leap, Inc. (COMPANY) All Rights Reserved.
-// Magic Leap, Inc. Confidential and Proprietary
-//
-// NOTICE:  All information contained herein is, and remains the property
-// of COMPANY. The intellectual and technical concepts contained herein
-// are proprietary to COMPANY and may be covered by U.S. and Foreign
-// Patents, patents in process, and are protected by trade secret or
-// copyright law.  Dissemination of this information or reproduction of
-// this material is strictly forbidden unless prior written permission is
-// obtained from COMPANY.  Access to the source code contained herein is
-// hereby forbidden to anyone except current COMPANY employees, managers
-// or contractors who have executed Confidentiality and Non-disclosure
-// agreements explicitly covering such access.
-//
-// The copyright notice above does not evidence any actual or intended
-// publication or disclosure  of  this source code, which includes
-// information that is confidential and/or proprietary, and is a trade
-// secret, of  COMPANY.   ANY REPRODUCTION, MODIFICATION, DISTRIBUTION,
-// PUBLIC  PERFORMANCE, OR PUBLIC DISPLAY OF OR THROUGH USE  OF THIS
-// SOURCE CODE  WITHOUT THE EXPRESS WRITTEN CONSENT OF COMPANY IS
-// STRICTLY PROHIBITED, AND IN VIOLATION OF APPLICABLE LAWS AND
-// INTERNATIONAL TREATIES.  THE RECEIPT OR POSSESSION OF  THIS SOURCE
-// CODE AND/OR RELATED INFORMATION DOES NOT CONVEY OR IMPLY ANY RIGHTS
-// TO REPRODUCE, DISCLOSE OR DISTRIBUTE ITS CONTENTS, OR TO MANUFACTURE,
-// USE, OR SELL ANYTHING THAT IT  MAY DESCRIBE, IN WHOLE OR IN PART.
-//
-// %COPYRIGHT_END%
-// --------------------------------------------------------------------*/
-// %BANNER_END%
+// Copyright 1998-2018 Epic Games, Inc. All Rights Reserved.
 
 #include "MagicLeapCameraPreviewPlayer.h"
 #include "IMagicLeapMediaModule.h"
diff --git a/Engine/Plugins/Lumin/MagicLeapMedia/Source/MagicLeapMedia/Private/MagicLeapCameraPreviewPlayer.h b/Engine/Plugins/Lumin/MagicLeapMedia/Source/MagicLeapMedia/Private/MagicLeapCameraPreviewPlayer.h
--- a/Engine/Plugins/Lumin/MagicLeapMedia/Source/MagicLeapMedia/Private/MagicLeapCameraPreviewPlayer.h
+++ b/Engine/Plugins/Lumin/MagicLeapMedia/Source/MagicLeapMedia/Private/MagicLeapCameraPreviewPlayer.h
@@ -1,36 +1,4 @@
-// %BANNER_BEGIN%
-// ---------------------------------------------------------------------
-// %COPYRIGHT_BEGIN%
-//
-// Copyright (c) 2017 Magic Leap, Inc. (COMPANY) All Rights Reserved.
-// Magic Leap, Inc. Confidential and Proprietary
-//
-// NOTICE:  All information contained herein is, and remains the property
-// of COMPANY. The intellectual and technical concepts contained herein
-// are proprietary to COMPANY and may be covered by U.S. and Foreign
-// Patents, patents in process, and are protected by trade secret or
-// copyright law.  Dissemination of this information or reproduction of
-// this material is strictly forbidden unless prior written permission is
-// obtained from COMPANY.  Access to the source code contained herein is
-// hereby forbidden to anyone except current COMPANY employees, managers
-// or contractors who have executed Confidentiality and Non-disclosure
-// agreements explicitly covering such access.
-//
-// The copyright notice above does not evidence any actual or intended
-// publication or disclosure  of  this source code, which includes
-// information that is confidential and/or proprietary, and is a trade
-// secret, of  COMPANY.   ANY REPRODUCTION, MODIFICATION, DISTRIBUTION,
-// PUBLIC  PERFORMANCE, OR PUBLIC DISPLAY OF OR THROUGH USE  OF THIS
-// SOURCE CODE  WITHOUT THE EXPRESS WRITTEN CONSENT OF COMPANY IS
-// STRICTLY PROHIBITED, AND IN VIOLATION OF APPLICABLE LAWS AND
-// INTERNATIONAL TREATIES.  THE RECEIPT OR POSSESSION OF  THIS SOURCE
-// CODE AND/OR RELATED INFORMATION DOES NOT CONVEY OR IMPLY ANY RIGHTS
-// TO REPRODUCE, DISCLOSE OR DISTRIBUTE ITS CONTENTS, OR TO MANUFACTURE,
-// USE, OR SELL ANYTHING THAT IT  MAY DESCRIBE, IN WHOLE OR IN PART.
-//
-// %COPYRIGHT_END%
-// --------------------------------------------------------------------*/
-// %BANNER_END%
+// Copyright 1998-2018 Epic Games, Inc. All Rights Reserved.
 
 #pragma once
 
diff --git a/Engine/Plugins/Lumin/MagicLeapMedia/Source/MagicLeapMedia/Private/MagicLeapMediaModule.cpp b/Engine/Plugins/Lumin/MagicLeapMedia/Source/MagicLeapMedia/Private/MagicLeapMediaModule.cpp
--- a/Engine/Plugins/Lumin/MagicLeapMedia/Source/MagicLeapMedia/Private/MagicLeapMediaModule.cpp
+++ b/Engine/Plugins/Lumin/MagicLeapMedia/Source/MagicLeapMedia/Private/MagicLeapMediaModule.cpp
@@ -1,36 +1,4 @@
-// %BANNER_BEGIN%
-// ---------------------------------------------------------------------
-// %COPYRIGHT_BEGIN%
-//
-// Copyright (c) 2017 Magic Leap, Inc. (COMPANY) All Rights Reserved.
-// Magic Leap, Inc. Confidential and Proprietary
-//
-// NOTICE:  All information contained herein is, and remains the property
-// of COMPANY. The intellectual and technical concepts contained herein
-// are proprietary to COMPANY and may be covered by U.S. and Foreign
-// Patents, patents in process, and are protected by trade secret or
-// copyright law.  Dissemination of this information or reproduction of
-// this material is strictly forbidden unless prior written permission is
-// obtained from COMPANY.  Access to the source code contained herein is
-// hereby forbidden to anyone except current COMPANY employees, managers
-// or contractors who have executed Confidentiality and Non-disclosure
-// agreements explicitly covering such access.
-//
-// The copyright notice above does not evidence any actual or intended
-// publication or disclosure  of  this source code, which includes
-// information that is confidential and/or proprietary, and is a trade
-// secret, of  COMPANY.   ANY REPRODUCTION, MODIFICATION, DISTRIBUTION,
-// PUBLIC  PERFORMANCE, OR PUBLIC DISPLAY OF OR THROUGH USE  OF THIS
-// SOURCE CODE  WITHOUT THE EXPRESS WRITTEN CONSENT OF COMPANY IS
-// STRICTLY PROHIBITED, AND IN VIOLATION OF APPLICABLE LAWS AND
-// INTERNATIONAL TREATIES.  THE RECEIPT OR POSSESSION OF  THIS SOURCE
-// CODE AND/OR RELATED INFORMATION DOES NOT CONVEY OR IMPLY ANY RIGHTS
-// TO REPRODUCE, DISCLOSE OR DISTRIBUTE ITS CONTENTS, OR TO MANUFACTURE,
-// USE, OR SELL ANYTHING THAT IT  MAY DESCRIBE, IN WHOLE OR IN PART.
-//
-// %COPYRIGHT_END%
-// --------------------------------------------------------------------*/
-// %BANNER_END%
+// Copyright 1998-2018 Epic Games, Inc. All Rights Reserved.
 
 #include "CoreMinimal.h"
 #include "MagicLeapMediaPlayer.h"
diff --git a/Engine/Plugins/Lumin/MagicLeapMedia/Source/MagicLeapMedia/Private/MagicLeapMediaPlayer.cpp b/Engine/Plugins/Lumin/MagicLeapMedia/Source/MagicLeapMedia/Private/MagicLeapMediaPlayer.cpp
--- a/Engine/Plugins/Lumin/MagicLeapMedia/Source/MagicLeapMedia/Private/MagicLeapMediaPlayer.cpp
+++ b/Engine/Plugins/Lumin/MagicLeapMedia/Source/MagicLeapMedia/Private/MagicLeapMediaPlayer.cpp
@@ -1,36 +1,4 @@
-// %BANNER_BEGIN%
-// ---------------------------------------------------------------------
-// %COPYRIGHT_BEGIN%
-//
-// Copyright (c) 2017 Magic Leap, Inc. (COMPANY) All Rights Reserved.
-// Magic Leap, Inc. Confidential and Proprietary
-//
-// NOTICE:  All information contained herein is, and remains the property
-// of COMPANY. The intellectual and technical concepts contained herein
-// are proprietary to COMPANY and may be covered by U.S. and Foreign
-// Patents, patents in process, and are protected by trade secret or
-// copyright law.  Dissemination of this information or reproduction of
-// this material is strictly forbidden unless prior written permission is
-// obtained from COMPANY.  Access to the source code contained herein is
-// hereby forbidden to anyone except current COMPANY employees, managers
-// or contractors who have executed Confidentiality and Non-disclosure
-// agreements explicitly covering such access.
-//
-// The copyright notice above does not evidence any actual or intended
-// publication or disclosure  of  this source code, which includes
-// information that is confidential and/or proprietary, and is a trade
-// secret, of  COMPANY.   ANY REPRODUCTION, MODIFICATION, DISTRIBUTION,
-// PUBLIC  PERFORMANCE, OR PUBLIC DISPLAY OF OR THROUGH USE  OF THIS
-// SOURCE CODE  WITHOUT THE EXPRESS WRITTEN CONSENT OF COMPANY IS
-// STRICTLY PROHIBITED, AND IN VIOLATION OF APPLICABLE LAWS AND
-// INTERNATIONAL TREATIES.  THE RECEIPT OR POSSESSION OF  THIS SOURCE
-// CODE AND/OR RELATED INFORMATION DOES NOT CONVEY OR IMPLY ANY RIGHTS
-// TO REPRODUCE, DISCLOSE OR DISTRIBUTE ITS CONTENTS, OR TO MANUFACTURE,
-// USE, OR SELL ANYTHING THAT IT  MAY DESCRIBE, IN WHOLE OR IN PART.
-//
-// %COPYRIGHT_END%
-// --------------------------------------------------------------------*/
-// %BANNER_END%
+// Copyright 1998-2018 Epic Games, Inc. All Rights Reserved.
 
 #include "MagicLeapMediaPlayer.h"
 #include "IMagicLeapMediaModule.h"
diff --git a/Engine/Plugins/Lumin/MagicLeapMedia/Source/MagicLeapMedia/Private/MagicLeapMediaPlayer.h b/Engine/Plugins/Lumin/MagicLeapMedia/Source/MagicLeapMedia/Private/MagicLeapMediaPlayer.h
--- a/Engine/Plugins/Lumin/MagicLeapMedia/Source/MagicLeapMedia/Private/MagicLeapMediaPlayer.h
+++ b/Engine/Plugins/Lumin/MagicLeapMedia/Source/MagicLeapMedia/Private/MagicLeapMediaPlayer.h
@@ -1,36 +1,4 @@
-// %BANNER_BEGIN%
-// ---------------------------------------------------------------------
-// %COPYRIGHT_BEGIN%
-//
-// Copyright (c) 2017 Magic Leap, Inc. (COMPANY) All Rights Reserved.
-// Magic Leap, Inc. Confidential and Proprietary
-//
-// NOTICE:  All information contained herein is, and remains the property
-// of COMPANY. The intellectual and technical concepts contained herein
-// are proprietary to COMPANY and may be covered by U.S. and Foreign
-// Patents, patents in process, and are protected by trade secret or
-// copyright law.  Dissemination of this information or reproduction of
-// this material is strictly forbidden unless prior written permission is
-// obtained from COMPANY.  Access to the source code contained herein is
-// hereby forbidden to anyone except current COMPANY employees, managers
-// or contractors who have executed Confidentiality and Non-disclosure
-// agreements explicitly covering such access.
-//
-// The copyright notice above does not evidence any actual or intended
-// publication or disclosure  of  this source code, which includes
-// information that is confidential and/or proprietary, and is a trade
-// secret, of  COMPANY.   ANY REPRODUCTION, MODIFICATION, DISTRIBUTION,
-// PUBLIC  PERFORMANCE, OR PUBLIC DISPLAY OF OR THROUGH USE  OF THIS
-// SOURCE CODE  WITHOUT THE EXPRESS WRITTEN CONSENT OF COMPANY IS
-// STRICTLY PROHIBITED, AND IN VIOLATION OF APPLICABLE LAWS AND
-// INTERNATIONAL TREATIES.  THE RECEIPT OR POSSESSION OF  THIS SOURCE
-// CODE AND/OR RELATED INFORMATION DOES NOT CONVEY OR IMPLY ANY RIGHTS
-// TO REPRODUCE, DISCLOSE OR DISTRIBUTE ITS CONTENTS, OR TO MANUFACTURE,
-// USE, OR SELL ANYTHING THAT IT  MAY DESCRIBE, IN WHOLE OR IN PART.
-//
-// %COPYRIGHT_END%
-// --------------------------------------------------------------------*/
-// %BANNER_END%
+// Copyright 1998-2018 Epic Games, Inc. All Rights Reserved.
 
 #pragma once
 
diff --git a/Engine/Plugins/Lumin/MagicLeapMedia/Source/MagicLeapMedia/Private/MagicLeapMediaTextureSample.h b/Engine/Plugins/Lumin/MagicLeapMedia/Source/MagicLeapMedia/Private/MagicLeapMediaTextureSample.h
--- a/Engine/Plugins/Lumin/MagicLeapMedia/Source/MagicLeapMedia/Private/MagicLeapMediaTextureSample.h
+++ b/Engine/Plugins/Lumin/MagicLeapMedia/Source/MagicLeapMedia/Private/MagicLeapMediaTextureSample.h
@@ -1,36 +1,4 @@
-// %BANNER_BEGIN%
-// ---------------------------------------------------------------------
-// %COPYRIGHT_BEGIN%
-//
-// Copyright (c) 2017 Magic Leap, Inc. (COMPANY) All Rights Reserved.
-// Magic Leap, Inc. Confidential and Proprietary
-//
-// NOTICE:  All information contained herein is, and remains the property
-// of COMPANY. The intellectual and technical concepts contained herein
-// are proprietary to COMPANY and may be covered by U.S. and Foreign
-// Patents, patents in process, and are protected by trade secret or
-// copyright law.  Dissemination of this information or reproduction of
-// this material is strictly forbidden unless prior written permission is
-// obtained from COMPANY.  Access to the source code contained herein is
-// hereby forbidden to anyone except current COMPANY employees, managers
-// or contractors who have executed Confidentiality and Non-disclosure
-// agreements explicitly covering such access.
-//
-// The copyright notice above does not evidence any actual or intended
-// publication or disclosure  of  this source code, which includes
-// information that is confidential and/or proprietary, and is a trade
-// secret, of  COMPANY.   ANY REPRODUCTION, MODIFICATION, DISTRIBUTION,
-// PUBLIC  PERFORMANCE, OR PUBLIC DISPLAY OF OR THROUGH USE  OF THIS
-// SOURCE CODE  WITHOUT THE EXPRESS WRITTEN CONSENT OF COMPANY IS
-// STRICTLY PROHIBITED, AND IN VIOLATION OF APPLICABLE LAWS AND
-// INTERNATIONAL TREATIES.  THE RECEIPT OR POSSESSION OF  THIS SOURCE
-// CODE AND/OR RELATED INFORMATION DOES NOT CONVEY OR IMPLY ANY RIGHTS
-// TO REPRODUCE, DISCLOSE OR DISTRIBUTE ITS CONTENTS, OR TO MANUFACTURE,
-// USE, OR SELL ANYTHING THAT IT  MAY DESCRIBE, IN WHOLE OR IN PART.
-//
-// %COPYRIGHT_END%
-// --------------------------------------------------------------------*/
-// %BANNER_END%
+// Copyright 1998-2018 Epic Games, Inc. All Rights Reserved.
 
 #pragma once
 
diff --git a/Engine/Plugins/Lumin/MagicLeapMedia/Source/MagicLeapMedia/Private/MediaWorker.cpp b/Engine/Plugins/Lumin/MagicLeapMedia/Source/MagicLeapMedia/Private/MediaWorker.cpp
--- a/Engine/Plugins/Lumin/MagicLeapMedia/Source/MagicLeapMedia/Private/MediaWorker.cpp
+++ b/Engine/Plugins/Lumin/MagicLeapMedia/Source/MagicLeapMedia/Private/MediaWorker.cpp
@@ -1,36 +1,4 @@
-// %BANNER_BEGIN%
-// ---------------------------------------------------------------------
-// %COPYRIGHT_BEGIN%
-//
-// Copyright (c) 2017 Magic Leap, Inc. (COMPANY) All Rights Reserved.
-// Magic Leap, Inc. Confidential and Proprietary
-//
-// NOTICE:  All information contained herein is, and remains the property
-// of COMPANY. The intellectual and technical concepts contained herein
-// are proprietary to COMPANY and may be covered by U.S. and Foreign
-// Patents, patents in process, and are protected by trade secret or
-// copyright law.  Dissemination of this information or reproduction of
-// this material is strictly forbidden unless prior written permission is
-// obtained from COMPANY.  Access to the source code contained herein is
-// hereby forbidden to anyone except current COMPANY employees, managers
-// or contractors who have executed Confidentiality and Non-disclosure
-// agreements explicitly covering such access.
-//
-// The copyright notice above does not evidence any actual or intended
-// publication or disclosure  of  this source code, which includes
-// information that is confidential and/or proprietary, and is a trade
-// secret, of  COMPANY.   ANY REPRODUCTION, MODIFICATION, DISTRIBUTION,
-// PUBLIC  PERFORMANCE, OR PUBLIC DISPLAY OF OR THROUGH USE  OF THIS
-// SOURCE CODE  WITHOUT THE EXPRESS WRITTEN CONSENT OF COMPANY IS
-// STRICTLY PROHIBITED, AND IN VIOLATION OF APPLICABLE LAWS AND
-// INTERNATIONAL TREATIES.  THE RECEIPT OR POSSESSION OF  THIS SOURCE
-// CODE AND/OR RELATED INFORMATION DOES NOT CONVEY OR IMPLY ANY RIGHTS
-// TO REPRODUCE, DISCLOSE OR DISTRIBUTE ITS CONTENTS, OR TO MANUFACTURE,
-// USE, OR SELL ANYTHING THAT IT  MAY DESCRIBE, IN WHOLE OR IN PART.
-//
-// %COPYRIGHT_END%
-// --------------------------------------------------------------------*/
-// %BANNER_END%
+// Copyright 1998-2018 Epic Games, Inc. All Rights Reserved.
 
 #include "MediaWorker.h"
 #include "IMagicLeapMediaModule.h"
diff --git a/Engine/Plugins/Lumin/MagicLeapMedia/Source/MagicLeapMedia/Private/MediaWorker.h b/Engine/Plugins/Lumin/MagicLeapMedia/Source/MagicLeapMedia/Private/MediaWorker.h
--- a/Engine/Plugins/Lumin/MagicLeapMedia/Source/MagicLeapMedia/Private/MediaWorker.h
+++ b/Engine/Plugins/Lumin/MagicLeapMedia/Source/MagicLeapMedia/Private/MediaWorker.h
@@ -1,36 +1,4 @@
-// %BANNER_BEGIN%
-// ---------------------------------------------------------------------
-// %COPYRIGHT_BEGIN%
-//
-// Copyright (c) 2017 Magic Leap, Inc. (COMPANY) All Rights Reserved.
-// Magic Leap, Inc. Confidential and Proprietary
-//
-// NOTICE:  All information contained herein is, and remains the property
-// of COMPANY. The intellectual and technical concepts contained herein
-// are proprietary to COMPANY and may be covered by U.S. and Foreign
-// Patents, patents in process, and are protected by trade secret or
-// copyright law.  Dissemination of this information or reproduction of
-// this material is strictly forbidden unless prior written permission is
-// obtained from COMPANY.  Access to the source code contained herein is
-// hereby forbidden to anyone except current COMPANY employees, managers
-// or contractors who have executed Confidentiality and Non-disclosure
-// agreements explicitly covering such access.
-//
-// The copyright notice above does not evidence any actual or intended
-// publication or disclosure  of  this source code, which includes
-// information that is confidential and/or proprietary, and is a trade
-// secret, of  COMPANY.   ANY REPRODUCTION, MODIFICATION, DISTRIBUTION,
-// PUBLIC  PERFORMANCE, OR PUBLIC DISPLAY OF OR THROUGH USE  OF THIS
-// SOURCE CODE  WITHOUT THE EXPRESS WRITTEN CONSENT OF COMPANY IS
-// STRICTLY PROHIBITED, AND IN VIOLATION OF APPLICABLE LAWS AND
-// INTERNATIONAL TREATIES.  THE RECEIPT OR POSSESSION OF  THIS SOURCE
-// CODE AND/OR RELATED INFORMATION DOES NOT CONVEY OR IMPLY ANY RIGHTS
-// TO REPRODUCE, DISCLOSE OR DISTRIBUTE ITS CONTENTS, OR TO MANUFACTURE,
-// USE, OR SELL ANYTHING THAT IT  MAY DESCRIBE, IN WHOLE OR IN PART.
-//
-// %COPYRIGHT_END%
-// --------------------------------------------------------------------*/
-// %BANNER_END%
+// Copyright 1998-2018 Epic Games, Inc. All Rights Reserved.
 
 #pragma once
 
diff --git a/Engine/Plugins/Lumin/MagicLeapMedia/Source/MagicLeapMedia/Public/IMagicLeapMediaModule.h b/Engine/Plugins/Lumin/MagicLeapMedia/Source/MagicLeapMedia/Public/IMagicLeapMediaModule.h
--- a/Engine/Plugins/Lumin/MagicLeapMedia/Source/MagicLeapMedia/Public/IMagicLeapMediaModule.h
+++ b/Engine/Plugins/Lumin/MagicLeapMedia/Source/MagicLeapMedia/Public/IMagicLeapMediaModule.h
@@ -1,36 +1,4 @@
-// %BANNER_BEGIN%
-// ---------------------------------------------------------------------
-// %COPYRIGHT_BEGIN%
-//
-// Copyright (c) 2017 Magic Leap, Inc. (COMPANY) All Rights Reserved.
-// Magic Leap, Inc. Confidential and Proprietary
-//
-// NOTICE:  All information contained herein is, and remains the property
-// of COMPANY. The intellectual and technical concepts contained herein
-// are proprietary to COMPANY and may be covered by U.S. and Foreign
-// Patents, patents in process, and are protected by trade secret or
-// copyright law.  Dissemination of this information or reproduction of
-// this material is strictly forbidden unless prior written permission is
-// obtained from COMPANY.  Access to the source code contained herein is
-// hereby forbidden to anyone except current COMPANY employees, managers
-// or contractors who have executed Confidentiality and Non-disclosure
-// agreements explicitly covering such access.
-//
-// The copyright notice above does not evidence any actual or intended
-// publication or disclosure  of  this source code, which includes
-// information that is confidential and/or proprietary, and is a trade
-// secret, of  COMPANY.   ANY REPRODUCTION, MODIFICATION, DISTRIBUTION,
-// PUBLIC  PERFORMANCE, OR PUBLIC DISPLAY OF OR THROUGH USE  OF THIS
-// SOURCE CODE  WITHOUT THE EXPRESS WRITTEN CONSENT OF COMPANY IS
-// STRICTLY PROHIBITED, AND IN VIOLATION OF APPLICABLE LAWS AND
-// INTERNATIONAL TREATIES.  THE RECEIPT OR POSSESSION OF  THIS SOURCE
-// CODE AND/OR RELATED INFORMATION DOES NOT CONVEY OR IMPLY ANY RIGHTS
-// TO REPRODUCE, DISCLOSE OR DISTRIBUTE ITS CONTENTS, OR TO MANUFACTURE,
-// USE, OR SELL ANYTHING THAT IT  MAY DESCRIBE, IN WHOLE OR IN PART.
-//
-// %COPYRIGHT_END%
-// --------------------------------------------------------------------*/
-// %BANNER_END%
+// Copyright 1998-2018 Epic Games, Inc. All Rights Reserved.
 
 #pragma once
 
diff --git a/Engine/Plugins/Lumin/MagicLeapMedia/Source/MagicLeapMediaEditor/MagicLeapMediaEditor.Build.cs b/Engine/Plugins/Lumin/MagicLeapMedia/Source/MagicLeapMediaEditor/MagicLeapMediaEditor.Build.cs
--- a/Engine/Plugins/Lumin/MagicLeapMedia/Source/MagicLeapMediaEditor/MagicLeapMediaEditor.Build.cs
+++ b/Engine/Plugins/Lumin/MagicLeapMedia/Source/MagicLeapMediaEditor/MagicLeapMediaEditor.Build.cs
@@ -1,36 +1,4 @@
-// %BANNER_BEGIN%
-// ---------------------------------------------------------------------
-// %COPYRIGHT_BEGIN%
-//
-// Copyright (c) 2017 Magic Leap, Inc. (COMPANY) All Rights Reserved.
-// Magic Leap, Inc. Confidential and Proprietary
-//
-// NOTICE: All information contained herein is, and remains the property
-// of COMPANY. The intellectual and technical concepts contained herein
-// are proprietary to COMPANY and may be covered by U.S. and Foreign
-// Patents, patents in process, and are protected by trade secret or
-// copyright law. Dissemination of this information or reproduction of
-// this material is strictly forbidden unless prior written permission is
-// obtained from COMPANY. Access to the source code contained herein is
-// hereby forbidden to anyone except current COMPANY employees, managers
-// or contractors who have executed Confidentiality and Non-disclosure
-// agreements explicitly covering such access.
-//
-// The copyright notice above does not evidence any actual or intended
-// publication or disclosure of this source code, which includes
-// information that is confidential and/or proprietary, and is a trade
-// secret, of COMPANY. ANY REPRODUCTION, MODIFICATION, DISTRIBUTION,
-// PUBLIC PERFORMANCE, OR PUBLIC DISPLAY OF OR THROUGH USE OF THIS
-// SOURCE CODE WITHOUT THE EXPRESS WRITTEN CONSENT OF COMPANY IS
-// STRICTLY PROHIBITED, AND IN VIOLATION OF APPLICABLE LAWS AND
-// INTERNATIONAL TREATIES. THE RECEIPT OR POSSESSION OF THIS SOURCE
-// CODE AND/OR RELATED INFORMATION DOES NOT CONVEY OR IMPLY ANY RIGHTS
-// TO REPRODUCE, DISCLOSE OR DISTRIBUTE ITS CONTENTS, OR TO MANUFACTURE,
-// USE, OR SELL ANYTHING THAT IT MAY DESCRIBE, IN WHOLE OR IN PART.
-//
-// %COPYRIGHT_END%
-// --------------------------------------------------------------------
-// %BANNER_END%
+// Copyright 1998-2018 Epic Games, Inc. All Rights Reserved.
 
 namespace UnrealBuildTool.Rules
 {
diff --git a/Engine/Plugins/Lumin/MagicLeapMedia/Source/MagicLeapMediaEditor/Private/MagicLeapFileMediaSourceFactory.cpp b/Engine/Plugins/Lumin/MagicLeapMedia/Source/MagicLeapMediaEditor/Private/MagicLeapFileMediaSourceFactory.cpp
--- a/Engine/Plugins/Lumin/MagicLeapMedia/Source/MagicLeapMediaEditor/Private/MagicLeapFileMediaSourceFactory.cpp
+++ b/Engine/Plugins/Lumin/MagicLeapMedia/Source/MagicLeapMediaEditor/Private/MagicLeapFileMediaSourceFactory.cpp
@@ -1,36 +1,4 @@
-// %BANNER_BEGIN%
-// ---------------------------------------------------------------------
-// %COPYRIGHT_BEGIN%
-//
-// Copyright (c) 2017 Magic Leap, Inc. (COMPANY) All Rights Reserved.
-// Magic Leap, Inc. Confidential and Proprietary
-//
-// NOTICE:  All information contained herein is, and remains the property
-// of COMPANY. The intellectual and technical concepts contained herein
-// are proprietary to COMPANY and may be covered by U.S. and Foreign
-// Patents, patents in process, and are protected by trade secret or
-// copyright law.  Dissemination of this information or reproduction of
-// this material is strictly forbidden unless prior written permission is
-// obtained from COMPANY.  Access to the source code contained herein is
-// hereby forbidden to anyone except current COMPANY employees, managers
-// or contractors who have executed Confidentiality and Non-disclosure
-// agreements explicitly covering such access.
-//
-// The copyright notice above does not evidence any actual or intended
-// publication or disclosure  of  this source code, which includes
-// information that is confidential and/or proprietary, and is a trade
-// secret, of  COMPANY.   ANY REPRODUCTION, MODIFICATION, DISTRIBUTION,
-// PUBLIC  PERFORMANCE, OR PUBLIC DISPLAY OF OR THROUGH USE  OF THIS
-// SOURCE CODE  WITHOUT THE EXPRESS WRITTEN CONSENT OF COMPANY IS
-// STRICTLY PROHIBITED, AND IN VIOLATION OF APPLICABLE LAWS AND
-// INTERNATIONAL TREATIES.  THE RECEIPT OR POSSESSION OF  THIS SOURCE
-// CODE AND/OR RELATED INFORMATION DOES NOT CONVEY OR IMPLY ANY RIGHTS
-// TO REPRODUCE, DISCLOSE OR DISTRIBUTE ITS CONTENTS, OR TO MANUFACTURE,
-// USE, OR SELL ANYTHING THAT IT  MAY DESCRIBE, IN WHOLE OR IN PART.
-//
-// %COPYRIGHT_END%
-// --------------------------------------------------------------------*/
-// %BANNER_END%
+// Copyright 1998-2018 Epic Games, Inc. All Rights Reserved.
 
 #include "MagicLeapFileMediaSourceFactory.h"
 #include "FileMediaSource.h"
diff --git a/Engine/Plugins/Lumin/MagicLeapMedia/Source/MagicLeapMediaEditor/Private/MagicLeapFileMediaSourceFactory.h b/Engine/Plugins/Lumin/MagicLeapMedia/Source/MagicLeapMediaEditor/Private/MagicLeapFileMediaSourceFactory.h
--- a/Engine/Plugins/Lumin/MagicLeapMedia/Source/MagicLeapMediaEditor/Private/MagicLeapFileMediaSourceFactory.h
+++ b/Engine/Plugins/Lumin/MagicLeapMedia/Source/MagicLeapMediaEditor/Private/MagicLeapFileMediaSourceFactory.h
@@ -1,36 +1,4 @@
-// %BANNER_BEGIN%
-// ---------------------------------------------------------------------
-// %COPYRIGHT_BEGIN%
-//
-// Copyright (c) 2017 Magic Leap, Inc. (COMPANY) All Rights Reserved.
-// Magic Leap, Inc. Confidential and Proprietary
-//
-// NOTICE:  All information contained herein is, and remains the property
-// of COMPANY. The intellectual and technical concepts contained herein
-// are proprietary to COMPANY and may be covered by U.S. and Foreign
-// Patents, patents in process, and are protected by trade secret or
-// copyright law.  Dissemination of this information or reproduction of
-// this material is strictly forbidden unless prior written permission is
-// obtained from COMPANY.  Access to the source code contained herein is
-// hereby forbidden to anyone except current COMPANY employees, managers
-// or contractors who have executed Confidentiality and Non-disclosure
-// agreements explicitly covering such access.
-//
-// The copyright notice above does not evidence any actual or intended
-// publication or disclosure  of  this source code, which includes
-// information that is confidential and/or proprietary, and is a trade
-// secret, of  COMPANY.   ANY REPRODUCTION, MODIFICATION, DISTRIBUTION,
-// PUBLIC  PERFORMANCE, OR PUBLIC DISPLAY OF OR THROUGH USE  OF THIS
-// SOURCE CODE  WITHOUT THE EXPRESS WRITTEN CONSENT OF COMPANY IS
-// STRICTLY PROHIBITED, AND IN VIOLATION OF APPLICABLE LAWS AND
-// INTERNATIONAL TREATIES.  THE RECEIPT OR POSSESSION OF  THIS SOURCE
-// CODE AND/OR RELATED INFORMATION DOES NOT CONVEY OR IMPLY ANY RIGHTS
-// TO REPRODUCE, DISCLOSE OR DISTRIBUTE ITS CONTENTS, OR TO MANUFACTURE,
-// USE, OR SELL ANYTHING THAT IT  MAY DESCRIBE, IN WHOLE OR IN PART.
-//
-// %COPYRIGHT_END%
-// --------------------------------------------------------------------*/
-// %BANNER_END%
+// Copyright 1998-2018 Epic Games, Inc. All Rights Reserved.
 
 #pragma once
 
diff --git a/Engine/Plugins/Lumin/MagicLeapMedia/Source/MagicLeapMediaEditor/Private/MagicLeapMediaEditorModule.cpp b/Engine/Plugins/Lumin/MagicLeapMedia/Source/MagicLeapMediaEditor/Private/MagicLeapMediaEditorModule.cpp
--- a/Engine/Plugins/Lumin/MagicLeapMedia/Source/MagicLeapMediaEditor/Private/MagicLeapMediaEditorModule.cpp
+++ b/Engine/Plugins/Lumin/MagicLeapMedia/Source/MagicLeapMediaEditor/Private/MagicLeapMediaEditorModule.cpp
@@ -1,36 +1,4 @@
-// %BANNER_BEGIN%
-// ---------------------------------------------------------------------
-// %COPYRIGHT_BEGIN%
-//
-// Copyright (c) 2017 Magic Leap, Inc. (COMPANY) All Rights Reserved.
-// Magic Leap, Inc. Confidential and Proprietary
-//
-// NOTICE:  All information contained herein is, and remains the property
-// of COMPANY. The intellectual and technical concepts contained herein
-// are proprietary to COMPANY and may be covered by U.S. and Foreign
-// Patents, patents in process, and are protected by trade secret or
-// copyright law.  Dissemination of this information or reproduction of
-// this material is strictly forbidden unless prior written permission is
-// obtained from COMPANY.  Access to the source code contained herein is
-// hereby forbidden to anyone except current COMPANY employees, managers
-// or contractors who have executed Confidentiality and Non-disclosure
-// agreements explicitly covering such access.
-//
-// The copyright notice above does not evidence any actual or intended
-// publication or disclosure  of  this source code, which includes
-// information that is confidential and/or proprietary, and is a trade
-// secret, of  COMPANY.   ANY REPRODUCTION, MODIFICATION, DISTRIBUTION,
-// PUBLIC  PERFORMANCE, OR PUBLIC DISPLAY OF OR THROUGH USE  OF THIS
-// SOURCE CODE  WITHOUT THE EXPRESS WRITTEN CONSENT OF COMPANY IS
-// STRICTLY PROHIBITED, AND IN VIOLATION OF APPLICABLE LAWS AND
-// INTERNATIONAL TREATIES.  THE RECEIPT OR POSSESSION OF  THIS SOURCE
-// CODE AND/OR RELATED INFORMATION DOES NOT CONVEY OR IMPLY ANY RIGHTS
-// TO REPRODUCE, DISCLOSE OR DISTRIBUTE ITS CONTENTS, OR TO MANUFACTURE,
-// USE, OR SELL ANYTHING THAT IT  MAY DESCRIBE, IN WHOLE OR IN PART.
-//
-// %COPYRIGHT_END%
-// --------------------------------------------------------------------*/
-// %BANNER_END%
+// Copyright 1998-2018 Epic Games, Inc. All Rights Reserved.
 
 #include "Modules/ModuleInterface.h"
 #include "Modules/ModuleManager.h"
diff --git a/Engine/Plugins/Lumin/MagicLeapMedia/Source/MagicLeapMediaFactory/MagicLeapMediaFactory.Build.cs b/Engine/Plugins/Lumin/MagicLeapMedia/Source/MagicLeapMediaFactory/MagicLeapMediaFactory.Build.cs
--- a/Engine/Plugins/Lumin/MagicLeapMedia/Source/MagicLeapMediaFactory/MagicLeapMediaFactory.Build.cs
+++ b/Engine/Plugins/Lumin/MagicLeapMedia/Source/MagicLeapMediaFactory/MagicLeapMediaFactory.Build.cs
@@ -1,36 +1,4 @@
-// %BANNER_BEGIN%
-// ---------------------------------------------------------------------
-// %COPYRIGHT_BEGIN%
-//
-// Copyright (c) 2017 Magic Leap, Inc. (COMPANY) All Rights Reserved.
-// Magic Leap, Inc. Confidential and Proprietary
-//
-// NOTICE:  All information contained herein is, and remains the property
-// of COMPANY. The intellectual and technical concepts contained herein
-// are proprietary to COMPANY and may be covered by U.S. and Foreign
-// Patents, patents in process, and are protected by trade secret or
-// copyright law.  Dissemination of this information or reproduction of
-// this material is strictly forbidden unless prior written permission is
-// obtained from COMPANY.  Access to the source code contained herein is
-// hereby forbidden to anyone except current COMPANY employees, managers
-// or contractors who have executed Confidentiality and Non-disclosure
-// agreements explicitly covering such access.
-//
-// The copyright notice above does not evidence any actual or intended
-// publication or disclosure  of  this source code, which includes
-// information that is confidential and/or proprietary, and is a trade
-// secret, of  COMPANY.   ANY REPRODUCTION, MODIFICATION, DISTRIBUTION,
-// PUBLIC  PERFORMANCE, OR PUBLIC DISPLAY OF OR THROUGH USE  OF THIS
-// SOURCE CODE  WITHOUT THE EXPRESS WRITTEN CONSENT OF COMPANY IS
-// STRICTLY PROHIBITED, AND IN VIOLATION OF APPLICABLE LAWS AND
-// INTERNATIONAL TREATIES.  THE RECEIPT OR POSSESSION OF  THIS SOURCE
-// CODE AND/OR RELATED INFORMATION DOES NOT CONVEY OR IMPLY ANY RIGHTS
-// TO REPRODUCE, DISCLOSE OR DISTRIBUTE ITS CONTENTS, OR TO MANUFACTURE,
-// USE, OR SELL ANYTHING THAT IT  MAY DESCRIBE, IN WHOLE OR IN PART.
-//
-// %COPYRIGHT_END%
-// --------------------------------------------------------------------*/
-// %BANNER_END%
+// Copyright 1998-2018 Epic Games, Inc. All Rights Reserved.
 
 namespace UnrealBuildTool.Rules
 {
diff --git a/Engine/Plugins/Lumin/MagicLeapMedia/Source/MagicLeapMediaFactory/Private/MagicLeapMediaFactoryModule.cpp b/Engine/Plugins/Lumin/MagicLeapMedia/Source/MagicLeapMediaFactory/Private/MagicLeapMediaFactoryModule.cpp
--- a/Engine/Plugins/Lumin/MagicLeapMedia/Source/MagicLeapMediaFactory/Private/MagicLeapMediaFactoryModule.cpp
+++ b/Engine/Plugins/Lumin/MagicLeapMedia/Source/MagicLeapMediaFactory/Private/MagicLeapMediaFactoryModule.cpp
@@ -1,36 +1,4 @@
-// %BANNER_BEGIN%
-// ---------------------------------------------------------------------
-// %COPYRIGHT_BEGIN%
-//
-// Copyright (c) 2017 Magic Leap, Inc. (COMPANY) All Rights Reserved.
-// Magic Leap, Inc. Confidential and Proprietary
-//
-// NOTICE:  All information contained herein is, and remains the property
-// of COMPANY. The intellectual and technical concepts contained herein
-// are proprietary to COMPANY and may be covered by U.S. and Foreign
-// Patents, patents in process, and are protected by trade secret or
-// copyright law.  Dissemination of this information or reproduction of
-// this material is strictly forbidden unless prior written permission is
-// obtained from COMPANY.  Access to the source code contained herein is
-// hereby forbidden to anyone except current COMPANY employees, managers
-// or contractors who have executed Confidentiality and Non-disclosure
-// agreements explicitly covering such access.
-//
-// The copyright notice above does not evidence any actual or intended
-// publication or disclosure  of  this source code, which includes
-// information that is confidential and/or proprietary, and is a trade
-// secret, of  COMPANY.   ANY REPRODUCTION, MODIFICATION, DISTRIBUTION,
-// PUBLIC  PERFORMANCE, OR PUBLIC DISPLAY OF OR THROUGH USE  OF THIS
-// SOURCE CODE  WITHOUT THE EXPRESS WRITTEN CONSENT OF COMPANY IS
-// STRICTLY PROHIBITED, AND IN VIOLATION OF APPLICABLE LAWS AND
-// INTERNATIONAL TREATIES.  THE RECEIPT OR POSSESSION OF  THIS SOURCE
-// CODE AND/OR RELATED INFORMATION DOES NOT CONVEY OR IMPLY ANY RIGHTS
-// TO REPRODUCE, DISCLOSE OR DISTRIBUTE ITS CONTENTS, OR TO MANUFACTURE,
-// USE, OR SELL ANYTHING THAT IT  MAY DESCRIBE, IN WHOLE OR IN PART.
-//
-// %COPYRIGHT_END%
-// --------------------------------------------------------------------*/
-// %BANNER_END%
+// Copyright 1998-2018 Epic Games, Inc. All Rights Reserved.
 
 #include "CoreMinimal.h"
 #include "Misc/Paths.h"
diff --git a/Engine/Plugins/Lumin/MagicLeapMedia/Source/MagicLeapMediaFactory/Private/MagicLeapMediaFactoryPrivate.h b/Engine/Plugins/Lumin/MagicLeapMedia/Source/MagicLeapMediaFactory/Private/MagicLeapMediaFactoryPrivate.h
--- a/Engine/Plugins/Lumin/MagicLeapMedia/Source/MagicLeapMediaFactory/Private/MagicLeapMediaFactoryPrivate.h
+++ b/Engine/Plugins/Lumin/MagicLeapMedia/Source/MagicLeapMediaFactory/Private/MagicLeapMediaFactoryPrivate.h
@@ -1,36 +1,4 @@
-// %BANNER_BEGIN%
-// ---------------------------------------------------------------------
-// %COPYRIGHT_BEGIN%
-//
-// Copyright (c) 2017 Magic Leap, Inc. (COMPANY) All Rights Reserved.
-// Magic Leap, Inc. Confidential and Proprietary
-//
-// NOTICE:  All information contained herein is, and remains the property
-// of COMPANY. The intellectual and technical concepts contained herein
-// are proprietary to COMPANY and may be covered by U.S. and Foreign
-// Patents, patents in process, and are protected by trade secret or
-// copyright law.  Dissemination of this information or reproduction of
-// this material is strictly forbidden unless prior written permission is
-// obtained from COMPANY.  Access to the source code contained herein is
-// hereby forbidden to anyone except current COMPANY employees, managers
-// or contractors who have executed Confidentiality and Non-disclosure
-// agreements explicitly covering such access.
-//
-// The copyright notice above does not evidence any actual or intended
-// publication or disclosure  of  this source code, which includes
-// information that is confidential and/or proprietary, and is a trade
-// secret, of  COMPANY.   ANY REPRODUCTION, MODIFICATION, DISTRIBUTION,
-// PUBLIC  PERFORMANCE, OR PUBLIC DISPLAY OF OR THROUGH USE  OF THIS
-// SOURCE CODE  WITHOUT THE EXPRESS WRITTEN CONSENT OF COMPANY IS
-// STRICTLY PROHIBITED, AND IN VIOLATION OF APPLICABLE LAWS AND
-// INTERNATIONAL TREATIES.  THE RECEIPT OR POSSESSION OF  THIS SOURCE
-// CODE AND/OR RELATED INFORMATION DOES NOT CONVEY OR IMPLY ANY RIGHTS
-// TO REPRODUCE, DISCLOSE OR DISTRIBUTE ITS CONTENTS, OR TO MANUFACTURE,
-// USE, OR SELL ANYTHING THAT IT  MAY DESCRIBE, IN WHOLE OR IN PART.
-//
-// %COPYRIGHT_END%
-// --------------------------------------------------------------------*/
-// %BANNER_END%
+// Copyright 1998-2018 Epic Games, Inc. All Rights Reserved.
 
 #pragma once
 
diff --git a/Engine/Plugins/Lumin/MagicLeapScreens/Source/MagicLeapScreens/Classes/ScreensComponent.h b/Engine/Plugins/Lumin/MagicLeapScreens/Source/MagicLeapScreens/Classes/ScreensComponent.h
--- a/Engine/Plugins/Lumin/MagicLeapScreens/Source/MagicLeapScreens/Classes/ScreensComponent.h
+++ b/Engine/Plugins/Lumin/MagicLeapScreens/Source/MagicLeapScreens/Classes/ScreensComponent.h
@@ -1,36 +1,4 @@
-// %BANNER_BEGIN%
-// ---------------------------------------------------------------------
-// %COPYRIGHT_BEGIN%
-//
-// Copyright (c) 2017 Magic Leap, Inc. (COMPANY) All Rights Reserved.
-// Magic Leap, Inc. Confidential and Proprietary
-//
-// NOTICE: All information contained herein is, and remains the property
-// of COMPANY. The intellectual and technical concepts contained herein
-// are proprietary to COMPANY and may be covered by U.S. and Foreign
-// Patents, patents in process, and are protected by trade secret or
-// copyright law. Dissemination of this information or reproduction of
-// this material is strictly forbidden unless prior written permission is
-// obtained from COMPANY. Access to the source code contained herein is
-// hereby forbidden to anyone except current COMPANY employees, managers
-// or contractors who have executed Confidentiality and Non-disclosure
-// agreements explicitly covering such access.
-//
-// The copyright notice above does not evidence any actual or intended
-// publication or disclosure of this source code, which includes
-// information that is confidential and/or proprietary, and is a trade
-// secret, of COMPANY. ANY REPRODUCTION, MODIFICATION, DISTRIBUTION,
-// PUBLIC PERFORMANCE, OR PUBLIC DISPLAY OF OR THROUGH USE OF THIS
-// SOURCE CODE WITHOUT THE EXPRESS WRITTEN CONSENT OF COMPANY IS
-// STRICTLY PROHIBITED, AND IN VIOLATION OF APPLICABLE LAWS AND
-// INTERNATIONAL TREATIES. THE RECEIPT OR POSSESSION OF THIS SOURCE
-// CODE AND/OR RELATED INFORMATION DOES NOT CONVEY OR IMPLY ANY RIGHTS
-// TO REPRODUCE, DISCLOSE OR DISTRIBUTE ITS CONTENTS, OR TO MANUFACTURE,
-// USE, OR SELL ANYTHING THAT IT MAY DESCRIBE, IN WHOLE OR IN PART.
-//
-// %COPYRIGHT_END%
-// --------------------------------------------------------------------
-// %BANNER_END%
+// Copyright 1998-2018 Epic Games, Inc. All Rights Reserved.
 
 #pragma once
 
diff --git a/Engine/Plugins/Lumin/MagicLeapScreens/Source/MagicLeapScreens/MagicLeapScreens.Build.cs b/Engine/Plugins/Lumin/MagicLeapScreens/Source/MagicLeapScreens/MagicLeapScreens.Build.cs
--- a/Engine/Plugins/Lumin/MagicLeapScreens/Source/MagicLeapScreens/MagicLeapScreens.Build.cs
+++ b/Engine/Plugins/Lumin/MagicLeapScreens/Source/MagicLeapScreens/MagicLeapScreens.Build.cs
@@ -1,36 +1,4 @@
-// %BANNER_BEGIN%
-// ---------------------------------------------------------------------
-// %COPYRIGHT_BEGIN%
-//
-// Copyright (c) 2017 Magic Leap, Inc. (COMPANY) All Rights Reserved.
-// Magic Leap, Inc. Confidential and Proprietary
-//
-// NOTICE: All information contained herein is, and remains the property
-// of COMPANY. The intellectual and technical concepts contained herein
-// are proprietary to COMPANY and may be covered by U.S. and Foreign
-// Patents, patents in process, and are protected by trade secret or
-// copyright law. Dissemination of this information or reproduction of
-// this material is strictly forbidden unless prior written permission is
-// obtained from COMPANY. Access to the source code contained herein is
-// hereby forbidden to anyone except current COMPANY employees, managers
-// or contractors who have executed Confidentiality and Non-disclosure
-// agreements explicitly covering such access.
-//
-// The copyright notice above does not evidence any actual or intended
-// publication or disclosure of this source code, which includes
-// information that is confidential and/or proprietary, and is a trade
-// secret, of COMPANY. ANY REPRODUCTION, MODIFICATION, DISTRIBUTION,
-// PUBLIC PERFORMANCE, OR PUBLIC DISPLAY OF OR THROUGH USE OF THIS
-// SOURCE CODE WITHOUT THE EXPRESS WRITTEN CONSENT OF COMPANY IS
-// STRICTLY PROHIBITED, AND IN VIOLATION OF APPLICABLE LAWS AND
-// INTERNATIONAL TREATIES. THE RECEIPT OR POSSESSION OF THIS SOURCE
-// CODE AND/OR RELATED INFORMATION DOES NOT CONVEY OR IMPLY ANY RIGHTS
-// TO REPRODUCE, DISCLOSE OR DISTRIBUTE ITS CONTENTS, OR TO MANUFACTURE,
-// USE, OR SELL ANYTHING THAT IT MAY DESCRIBE, IN WHOLE OR IN PART.
-//
-// %COPYRIGHT_END%
-// --------------------------------------------------------------------
-// %BANNER_END%
+// Copyright 1998-2018 Epic Games, Inc. All Rights Reserved.
 
 using System.IO;
 
diff --git a/Engine/Plugins/Lumin/MagicLeapScreens/Source/MagicLeapScreens/Private/ScreensComponent.cpp b/Engine/Plugins/Lumin/MagicLeapScreens/Source/MagicLeapScreens/Private/ScreensComponent.cpp
--- a/Engine/Plugins/Lumin/MagicLeapScreens/Source/MagicLeapScreens/Private/ScreensComponent.cpp
+++ b/Engine/Plugins/Lumin/MagicLeapScreens/Source/MagicLeapScreens/Private/ScreensComponent.cpp
@@ -1,4 +1,4 @@
-// Fill out your copyright notice in the Description page of Project Settings.
+// Copyright 1998-2018 Epic Games, Inc. All Rights Reserved.
 
 #include "ScreensComponent.h"
 #include "CoreMinimal.h"
diff --git a/Engine/Plugins/Lumin/MagicLeapScreens/Source/MagicLeapScreens/Public/IMagicLeapScreensPlugin.h b/Engine/Plugins/Lumin/MagicLeapScreens/Source/MagicLeapScreens/Public/IMagicLeapScreensPlugin.h
--- a/Engine/Plugins/Lumin/MagicLeapScreens/Source/MagicLeapScreens/Public/IMagicLeapScreensPlugin.h
+++ b/Engine/Plugins/Lumin/MagicLeapScreens/Source/MagicLeapScreens/Public/IMagicLeapScreensPlugin.h
@@ -1,36 +1,4 @@
-// %BANNER_BEGIN%
-// ---------------------------------------------------------------------
-// %COPYRIGHT_BEGIN%
-//
-// Copyright (c) 2017 Magic Leap, Inc. (COMPANY) All Rights Reserved.
-// Magic Leap, Inc. Confidential and Proprietary
-//
-// NOTICE: All information contained herein is, and remains the property
-// of COMPANY. The intellectual and technical concepts contained herein
-// are proprietary to COMPANY and may be covered by U.S. and Foreign
-// Patents, patents in process, and are protected by trade secret or
-// copyright law. Dissemination of this information or reproduction of
-// this material is strictly forbidden unless prior written permission is
-// obtained from COMPANY. Access to the source code contained herein is
-// hereby forbidden to anyone except current COMPANY employees, managers
-// or contractors who have executed Confidentiality and Non-disclosure
-// agreements explicitly covering such access.
-//
-// The copyright notice above does not evidence any actual or intended
-// publication or disclosure of this source code, which includes
-// information that is confidential and/or proprietary, and is a trade
-// secret, of COMPANY. ANY REPRODUCTION, MODIFICATION, DISTRIBUTION,
-// PUBLIC PERFORMANCE, OR PUBLIC DISPLAY OF OR THROUGH USE OF THIS
-// SOURCE CODE WITHOUT THE EXPRESS WRITTEN CONSENT OF COMPANY IS
-// STRICTLY PROHIBITED, AND IN VIOLATION OF APPLICABLE LAWS AND
-// INTERNATIONAL TREATIES. THE RECEIPT OR POSSESSION OF THIS SOURCE
-// CODE AND/OR RELATED INFORMATION DOES NOT CONVEY OR IMPLY ANY RIGHTS
-// TO REPRODUCE, DISCLOSE OR DISTRIBUTE ITS CONTENTS, OR TO MANUFACTURE,
-// USE, OR SELL ANYTHING THAT IT MAY DESCRIBE, IN WHOLE OR IN PART.
-//
-// %COPYRIGHT_END%
-// --------------------------------------------------------------------
-// %BANNER_END%
+// Copyright 1998-2018 Epic Games, Inc. All Rights Reserved.
 
 #pragma once
 
diff --git a/Engine/Source/Runtime/ApplicationCore/Private/Lumin/LuminInputInterface.cpp b/Engine/Source/Runtime/ApplicationCore/Private/Lumin/LuminInputInterface.cpp
--- a/Engine/Source/Runtime/ApplicationCore/Private/Lumin/LuminInputInterface.cpp
+++ b/Engine/Source/Runtime/ApplicationCore/Private/Lumin/LuminInputInterface.cpp
@@ -1,39 +1,5 @@
 // Copyright 1998-2018 Epic Games, Inc. All Rights Reserved.
 
-// %BANNER_BEGIN%
-// ---------------------------------------------------------------------
-// %COPYRIGHT_BEGIN%
-//
-// Copyright (c) 201x Magic Leap, Inc. (COMPANY) All Rights Reserved.
-// Magic Leap, Inc. Confidential and Proprietary
-//
-// NOTICE: All information contained herein is, and remains the property
-// of COMPANY. The intellectual and technical concepts contained herein
-// are proprietary to COMPANY and may be covered by U.S. and Foreign
-// Patents, patents in process, and are protected by trade secret or
-// copyright law. Dissemination of this information or reproduction of
-// this material is strictly forbidden unless prior written permission is
-// obtained from COMPANY. Access to the source code contained herein is
-// hereby forbidden to anyone except current COMPANY employees, managers
-// or contractors who have executed Confidentiality and Non-disclosure
-// agreements explicitly covering such access.
-//
-// The copyright notice above does not evidence any actual or intended
-// publication or disclosure of this source code, which includes
-// information that is confidential and/or proprietary, and is a trade
-// secret, of COMPANY. ANY REPRODUCTION, MODIFICATION, DISTRIBUTION,
-// PUBLIC PERFORMANCE, OR PUBLIC DISPLAY OF OR THROUGH USE OF THIS
-// SOURCE CODE WITHOUT THE EXPRESS WRITTEN CONSENT OF COMPANY IS
-// STRICTLY PROHIBITED, AND IN VIOLATION OF APPLICABLE LAWS AND
-// INTERNATIONAL TREATIES. THE RECEIPT OR POSSESSION OF THIS SOURCE
-// CODE AND/OR RELATED INFORMATION DOES NOT CONVEY OR IMPLY ANY RIGHTS
-// TO REPRODUCE, DISCLOSE OR DISTRIBUTE ITS CONTENTS, OR TO MANUFACTURE,
-// USE, OR SELL ANYTHING THAT IT MAY DESCRIBE, IN WHOLE OR IN PART.
-//
-// %COPYRIGHT_END%
-// --------------------------------------------------------------------
-// %BANNER_END%
-
 #include "Lumin/LuminInputInterface.h"
 #include "LuminGamepadInterface.h"
 #include "GenericPlatform/GenericApplication.h"
diff --git a/Engine/Source/Runtime/ApplicationCore/Public/Lumin/LuminInputInterface.h b/Engine/Source/Runtime/ApplicationCore/Public/Lumin/LuminInputInterface.h
--- a/Engine/Source/Runtime/ApplicationCore/Public/Lumin/LuminInputInterface.h
+++ b/Engine/Source/Runtime/ApplicationCore/Public/Lumin/LuminInputInterface.h
@@ -1,36 +1,4 @@
-// %BANNER_BEGIN%
-// ---------------------------------------------------------------------
-// %COPYRIGHT_BEGIN%
-//
-// Copyright (c) 201x Magic Leap, Inc. (COMPANY) All Rights Reserved.
-// Magic Leap, Inc. Confidential and Proprietary
-//
-// NOTICE: All information contained herein is, and remains the property
-// of COMPANY. The intellectual and technical concepts contained herein
-// are proprietary to COMPANY and may be covered by U.S. and Foreign
-// Patents, patents in process, and are protected by trade secret or
-// copyright law. Dissemination of this information or reproduction of
-// this material is strictly forbidden unless prior written permission is
-// obtained from COMPANY. Access to the source code contained herein is
-// hereby forbidden to anyone except current COMPANY employees, managers
-// or contractors who have executed Confidentiality and Non-disclosure
-// agreements explicitly covering such access.
-//
-// The copyright notice above does not evidence any actual or intended
-// publication or disclosure of this source code, which includes
-// information that is confidential and/or proprietary, and is a trade
-// secret, of COMPANY. ANY REPRODUCTION, MODIFICATION, DISTRIBUTION,
-// PUBLIC PERFORMANCE, OR PUBLIC DISPLAY OF OR THROUGH USE OF THIS
-// SOURCE CODE WITHOUT THE EXPRESS WRITTEN CONSENT OF COMPANY IS
-// STRICTLY PROHIBITED, AND IN VIOLATION OF APPLICABLE LAWS AND
-// INTERNATIONAL TREATIES. THE RECEIPT OR POSSESSION OF THIS SOURCE
-// CODE AND/OR RELATED INFORMATION DOES NOT CONVEY OR IMPLY ANY RIGHTS
-// TO REPRODUCE, DISCLOSE OR DISTRIBUTE ITS CONTENTS, OR TO MANUFACTURE,
-// USE, OR SELL ANYTHING THAT IT MAY DESCRIBE, IN WHOLE OR IN PART.
-//
-// %COPYRIGHT_END%
-// --------------------------------------------------------------------
-// %BANNER_END%
+// Copyright 1998-2018 Epic Games, Inc. All Rights Reserved.
 
 #pragma once
 
