#DEV-VR - Merge from //UE4/Dev-MagicLeap/... @ CL 4136411

diff --git a/Engine/Plugins/Lumin/MagicLeapMedia/Source/MagicLeapMedia/Private/MagicLeapMediaPlayer.cpp b/Engine/Plugins/Lumin/MagicLeapMedia/Source/MagicLeapMedia/Private/MagicLeapMediaPlayer.cpp
--- a/Engine/Plugins/Lumin/MagicLeapMedia/Source/MagicLeapMedia/Private/MagicLeapMediaPlayer.cpp
+++ b/Engine/Plugins/Lumin/MagicLeapMedia/Source/MagicLeapMedia/Private/MagicLeapMediaPlayer.cpp
@@ -132,8 +132,9 @@ FMagicLeapMediaPlayer::FMagicLeapMediaPlayer(IMediaEventSink& InEventSink)
 , MediaWorker(nullptr)
 , bWasMediaPlayingBeforeAppPause(false)
 {
-	MediaPlayerHandle = MLMediaPlayerCreate();
-	CurrentState = (Samples.IsValid() && MLHandleIsValid(MediaPlayerHandle)) ? EMediaState::Closed : EMediaState::Error;
+	MLResult Result = MLMediaPlayerCreate(&MediaPlayerHandle);
+	UE_CLOG(Result != MLResult_Ok, LogMagicLeapMedia, Error, TEXT("MLMediaPlayerCreate failed with error %d."), Result);
+	CurrentState = (Samples.IsValid() && Result == MLResult_Ok) ? EMediaState::Closed : EMediaState::Error;
 	MediaWorker = new FMediaWorker(MediaPlayerHandle, CriticalSection);
 }
 
@@ -182,22 +183,16 @@ FMagicLeapMediaPlayer::~FMagicLeapMediaPlayer()
 						Params.MediaPlayer->RenderThreadReleaseNativeBuffer(Params.MediaPlayerHandle, Params.TextureData->PreviousNativeBuffer);
 					}
 
-					bool bMediaPlayerDestroyed = MLMediaPlayerDestroy(Params.MediaPlayerHandle);
-					if (!bMediaPlayerDestroyed)
-					{
-						UE_LOG(LogMagicLeapMedia, Error, TEXT("Error destroying Magic Leap media player."));
-					}
+					MLResult Result = MLMediaPlayerDestroy(Params.MediaPlayerHandle);
+					UE_CLOG(Result != MLResult_Ok, LogMagicLeapMedia, Error, TEXT("MLMediaPlayerDestroy failed with error %d."), Result);
 				});
 			
 			FlushRenderingCommands();
 		}
 		else
 		{
-			bool bMediaPlayerDestroyed = MLMediaPlayerDestroy(MediaPlayerHandle);
-			if (!bMediaPlayerDestroyed)
-			{
-				UE_LOG(LogMagicLeapMedia, Error, TEXT("Error destroying Magic Leap media player."));
-			}
+			MLResult Result = MLMediaPlayerDestroy(MediaPlayerHandle);
+			UE_CLOG(Result != MLResult_Ok, LogMagicLeapMedia, Error, TEXT("MLMediaPlayerDestroy failed with error %d."), Result);
 		}
 
 		MediaPlayerHandle = ML_INVALID_HANDLE;
@@ -227,11 +222,11 @@ void FMagicLeapMediaPlayer::Close()
 		PauseHandle.Reset();
 	}
 
-	bool bStopped = MLMediaPlayerStop(MediaPlayerHandle);
-	if (!bStopped)
-	{
-		UE_LOG(LogMagicLeapMedia, Error, TEXT("Error stopping media player"));
-	}
+	MLResult StopResult = MLMediaPlayerStop(MediaPlayerHandle);
+	UE_CLOG(StopResult != MLResult_Ok, LogMagicLeapMedia, Error, TEXT("MLMediaPlayerStop failed with error %d"), StopResult);
+
+	MLResult ResetResult = MLMediaPlayerReset(MediaPlayerHandle);
+	UE_CLOG(ResetResult != MLResult_Ok, LogMagicLeapMedia, Error, TEXT("MLMediaPlayerReset failed with error %d"), ResetResult);
 
 	CurrentState = EMediaState::Closed;
 
@@ -239,6 +234,8 @@ void FMagicLeapMediaPlayer::Close()
 	Info.Empty();
 	MediaUrl = FString();
 	VideoSamplePool->Reset();
+	TrackInfo.Empty();
+	SelectedTrack.Empty();
 
 	// notify listeners
 	EventSink.ReceiveMediaEvent(EMediaEvent::TracksChanged);
@@ -328,10 +325,10 @@ bool FMagicLeapMediaPlayer::Open(const FString& Url, const IMediaOptions* Option
 			return false;
 		}
 
-		bool bMediaPlayerDataSource = MLMediaPlayerSetDataSourceForPath(MediaPlayerHandle, TCHAR_TO_UTF8(*FilePath));
-		if (!bMediaPlayerDataSource)
+		MLResult Result = MLMediaPlayerSetDataSourceForPath(MediaPlayerHandle, TCHAR_TO_UTF8(*FilePath));
+		if (Result != MLResult_Ok)
 		{
-			UE_LOG(LogMagicLeapMedia, Error, TEXT("Failed to set media player data source for path %s."), *FilePath);
+			UE_LOG(LogMagicLeapMedia, Error, TEXT("MLMediaPlayerSetDataSourceForPath for path %s failed with error %d."), *FilePath, Result);
 			EventSink.ReceiveMediaEvent(EMediaEvent::MediaOpenFailed);
 			return false;
 		}
@@ -339,10 +336,10 @@ bool FMagicLeapMediaPlayer::Open(const FString& Url, const IMediaOptions* Option
 	else
 	{    
 		// open remote media
-		bool bMediaPlayerDataSource = MLMediaPlayerSetDataSourceForURI(MediaPlayerHandle, TCHAR_TO_UTF8(*Url));
-		if (!bMediaPlayerDataSource)
+		MLResult Result = MLMediaPlayerSetDataSourceForURI(MediaPlayerHandle, TCHAR_TO_UTF8(*Url));
+		if (Result != MLResult_Ok)
 		{
-			UE_LOG(LogMagicLeapMedia, Error, TEXT("Error setting remote media source %s."), *Url);
+			UE_LOG(LogMagicLeapMedia, Error, TEXT("MLMediaPlayerSetDataSourceForURI for remote media source %s failed with error %d."), *Url, Result);
 			EventSink.ReceiveMediaEvent(EMediaEvent::MediaOpenFailed);
 			return false;
 		}
@@ -353,10 +350,10 @@ bool FMagicLeapMediaPlayer::Open(const FString& Url, const IMediaOptions* Option
 	// prepare media
 	MediaUrl = Url;
 
-	bool bResult = MLMediaPlayerPrepareAsync(MediaPlayerHandle);
-	if (!bResult)
+	MLResult Result = MLMediaPlayerPrepareAsync(MediaPlayerHandle);
+	if (Result != MLResult_Ok)
 	{
-		UE_LOG(LogMagicLeapMedia, Error, TEXT("Failed to prepare media source %s"), *Url);
+		UE_LOG(LogMagicLeapMedia, Error, TEXT("MLMediaPlayerPrepareAsync for media source %s failed with error %d"), *Url, Result);
 		EventSink.ReceiveMediaEvent(EMediaEvent::MediaOpenFailed);
 		return false;
 	}
@@ -394,15 +391,23 @@ bool FMagicLeapMediaPlayer::CanControl(EMediaControl Control) const
 
 FTimespan FMagicLeapMediaPlayer::GetDuration() const
 {
+	FTimespan Duration = FTimespan::Zero();
+
 	if (CurrentState == EMediaState::Playing || CurrentState == EMediaState::Paused || CurrentState == EMediaState::Stopped)
 	{
 		int32 durationMilSec = 0;
-		bool bResult = MLMediaPlayerGetDuration(MediaPlayerHandle, &durationMilSec);
-
-		return (bResult) ? FTimespan::FromMilliseconds(durationMilSec) : FTimespan::Zero();
+		MLResult Result = MLMediaPlayerGetDuration(MediaPlayerHandle, &durationMilSec);
+		if (Result == MLResult_Ok)
+		{
+			Duration = FTimespan::FromMilliseconds(durationMilSec);
+		}
+		else
+		{
+			UE_LOG(LogMagicLeapMedia, Error, TEXT("MLMediaPlayerGetDuration failed with error %d"), Result);
+		}
 	}
 
-	return FTimespan::Zero();
+	return Duration;
 }
 
 float FMagicLeapMediaPlayer::GetRate() const
@@ -432,15 +437,23 @@ TRangeSet<float> FMagicLeapMediaPlayer::GetSupportedRates(EMediaRateThinning Thi
 
 FTimespan FMagicLeapMediaPlayer::GetTime() const
 {
+	FTimespan Time = FTimespan::Zero();
+
 	if (CurrentState == EMediaState::Playing || CurrentState == EMediaState::Paused)
 	{
 		int32 currentPositionMilSec = 0;
-		bool bResult = MLMediaPlayerGetCurrentPosition(MediaPlayerHandle, &currentPositionMilSec);
-
-		return (bResult) ? FTimespan::FromMilliseconds(currentPositionMilSec) : FTimespan::Zero();
+		MLResult Result = MLMediaPlayerGetCurrentPosition(MediaPlayerHandle, &currentPositionMilSec);
+		if (Result == MLResult_Ok)
+		{
+			Time = FTimespan::FromMilliseconds(currentPositionMilSec);
+		}
+		else
+		{
+			UE_LOG(LogMagicLeapMedia, Error, TEXT("MLMediaPlayerGetCurrentPosition failed with error %d"), Result);
+		}
 	}
 
-	return FTimespan::Zero();
+	return Time;
 }
 
 bool FMagicLeapMediaPlayer::IsLooping() const
@@ -450,23 +463,36 @@ bool FMagicLeapMediaPlayer::IsLooping() const
 
 bool FMagicLeapMediaPlayer::Seek(const FTimespan& Time)
 {
+	bool bSuccess = true;
+
 	if ((CurrentState == EMediaState::Closed) || (CurrentState == EMediaState::Error) || (CurrentState == EMediaState::Preparing))
 	{
+		bSuccess = false;
 		UE_LOG(LogMagicLeapMedia, Warning, TEXT("Cannot seek while closed, preparing, or in error state"));
-		return false;
 	}
-
-	if (CurrentState == EMediaState::Playing || CurrentState == EMediaState::Paused)
+	else if (CurrentState == EMediaState::Playing || CurrentState == EMediaState::Paused)
 	{
-		return MLMediaPlayerSeekTo(MediaPlayerHandle, Time.GetTotalMilliseconds());
+		MLResult Result = MLMediaPlayerSeekTo(MediaPlayerHandle, Time.GetTotalMilliseconds(), MLMediaSeekMode_Previous_Sync);
+		if (Result != MLResult_Ok)
+		{
+			bSuccess = false;
+			UE_LOG(LogMagicLeapMedia, Error, TEXT("MLMediaPlayerSeekTo failed with error %d"), Result);
+		}
 	}
 
-	return false;
+	return bSuccess;
 }
 
 bool FMagicLeapMediaPlayer::SetLooping(bool Looping)
 {
-	return MLMediaPlayerSetLooping(MediaPlayerHandle, Looping);
+	MLResult Result = MLMediaPlayerSetLooping(MediaPlayerHandle, Looping);
+	if (Result != MLResult_Ok)
+	{
+		UE_LOG(LogMagicLeapMedia, Error, TEXT("MLMediaPlayerSetLooping failed with error %d"), Result);
+		return false;
+	}
+
+	return true;
 }
 
 bool FMagicLeapMediaPlayer::SetRate(float Rate)
@@ -486,12 +512,16 @@ bool FMagicLeapMediaPlayer::SetRate(float Rate)
 	bool bResult = false;
 	if (Rate == 0.0f)
 	{
-		bResult = MLMediaPlayerPause(MediaPlayerHandle);
-		if (bResult)
+		MLResult Result = MLMediaPlayerPause(MediaPlayerHandle);
+		if (Result == MLResult_Ok)
 		{
 			CurrentState = EMediaState::Paused;
 			EventSink.ReceiveMediaEvent(EMediaEvent::PlaybackSuspended);
 		}
+		else
+		{
+			UE_LOG(LogMagicLeapMedia, Error, TEXT("MLMediaPlayerPause failed with error %d!"), Result);
+		}
 	}
 	else if (Rate == 1.0f)
 	{
@@ -508,13 +538,20 @@ bool FMagicLeapMediaPlayer::SetRate(float Rate)
 
 bool FMagicLeapMediaPlayer::SetNativeVolume(float Volume)
 {
+	bool bSuccess = true;
+
 	if (MLHandleIsValid(MediaPlayerHandle))
 	{
 		Volume = (Volume < 0.0f) ? 0.0f : ((Volume < 1.0f) ? Volume : 1.0f);
-		return MLMediaPlayerSetVolume(MediaPlayerHandle, Volume);
+		MLResult Result = MLMediaPlayerSetVolume(MediaPlayerHandle, Volume);
+		if (Result != MLResult_Ok)
+		{
+			bSuccess = false;
+			UE_LOG(LogMagicLeapMedia, Error, TEXT("MLMediaPlayerSetVolume failed with error %d."), Result);
+		}
 	}
 
-	return false;
+	return bSuccess;
 }
 
 void FMagicLeapMediaPlayer::SetGuid(const FGuid& Guid)
@@ -582,6 +619,19 @@ void FMagicLeapMediaPlayer::TickFetch(FTimespan DeltaTime, FTimespan Timecode)
 					return;
 				}
 
+				MLMediaPlayerGetFrameTransformationMatrix(Params.MediaPlayerHandle, Params.MediaPlayer->FrameTransformationMatrix);
+				if (Params.MediaPlayer->UScale != Params.MediaPlayer->FrameTransformationMatrix[0] ||
+					Params.MediaPlayer->UOffset != Params.MediaPlayer->FrameTransformationMatrix[12] ||
+					(-Params.MediaPlayer->VScale) != Params.MediaPlayer->FrameTransformationMatrix[5] ||
+					(1.0f - Params.MediaPlayer->VOffset) != Params.MediaPlayer->FrameTransformationMatrix[13])
+				{
+					Params.MediaPlayer->UScale = Params.MediaPlayer->FrameTransformationMatrix[0];
+					Params.MediaPlayer->UOffset = Params.MediaPlayer->FrameTransformationMatrix[12];
+					Params.MediaPlayer->VScale = -Params.MediaPlayer->FrameTransformationMatrix[5];
+					Params.MediaPlayer->VOffset = 1.0f - Params.MediaPlayer->FrameTransformationMatrix[13];
+					TextureDataPtr->bIsVideoTextureValid = false;
+				}
+
 				int32 CurrentFramePosition = 0;
 				if (!Params.MediaPlayer->RenderThreadGetCurrentPosition(Params.MediaPlayerHandle, CurrentFramePosition))
 				{
@@ -635,7 +685,7 @@ void FMagicLeapMediaPlayer::TickFetch(FTimespan DeltaTime, FTimespan Timecode)
 				{
 					FSamplerStateInitializerRHI SamplerStateInitializer(SF_Bilinear, AM_Clamp, AM_Clamp, AM_Clamp);
 					FSamplerStateRHIRef SamplerStateRHI = RHICreateSamplerState(SamplerStateInitializer);
-					FExternalTextureRegistry::Get().RegisterExternalTexture(Params.PlayerGuid, MediaVideoTexture, SamplerStateRHI, FLinearColor(1.0f, 0.0f, 0.0f, 1.0f), FLinearColor(0.0f, 0.0f, 0.0f, 0.0f));
+					Params.MediaPlayer->RegisterExternalTexture(Params.PlayerGuid, MediaVideoTexture, SamplerStateRHI);
 
 					TextureDataPtr->bIsVideoTextureValid = true;
 				}
@@ -746,6 +796,11 @@ void FMagicLeapMediaPlayer::TickFetch(FTimespan DeltaTime, FTimespan Timecode)
 	}
 }
 
+void FMagicLeapMediaPlayer::RegisterExternalTexture(const FGuid& InGuid, FTextureRHIRef& InTextureRHI, FSamplerStateRHIRef& InSamplerStateRHI)
+{
+	FExternalTextureRegistry::Get().RegisterExternalTexture(InGuid, InTextureRHI, InSamplerStateRHI, FLinearColor(UScale, 0.0f, 0.0f, VScale), FLinearColor(UOffset,  VOffset, 0.0f, 0.0f));
+}
+
 void FMagicLeapMediaPlayer::TickInput(FTimespan DeltaTime, FTimespan Timecode)
 {
 	if (!bMediaPrepared)
@@ -905,13 +960,17 @@ FString FMagicLeapMediaPlayer::GetTrackLanguage(EMediaTrackType TrackType, int32
 	if (TrackInfo.Contains(TrackType) && TrackIndex >= 0 && TrackIndex < TrackInfo[TrackType].Num())
 	{
 		char* TrackLanguage = nullptr;
-		bool bResult = MLMediaPlayerGetTrackLanguage(MediaPlayerHandle, static_cast<uint32>(TrackInfo[TrackType][TrackIndex]), &TrackLanguage);
-		if (bResult)
+		MLResult Result = MLMediaPlayerGetTrackLanguage(MediaPlayerHandle, static_cast<uint32>(TrackInfo[TrackType][TrackIndex]), &TrackLanguage);
+		if (Result == MLResult_Ok)
 		{
 			FString Language(UTF8_TO_TCHAR(TrackLanguage));
 			free(TrackLanguage);
 			return Language;
 		}
+		else
+		{
+			UE_LOG(LogMagicLeapMedia, Error, TEXT("MLMediaPlayerGetTrackLanguage failed with error %d"), Result);
+		}
 	}
 
 	return FString();
@@ -932,28 +991,37 @@ bool FMagicLeapMediaPlayer::GetVideoTrackFormat(int32 TrackIndex, int32 FormatIn
 
 	int32 width = 0;
 	int32 height = 0;
-	bool bResult = MLMediaPlayerGetVideoSize(MediaPlayerHandle, &width, &height);
-	if (bResult)
+	MLResult Result = MLMediaPlayerGetVideoSize(MediaPlayerHandle, &width, &height);
+	if (Result != MLResult_Ok)
 	{
-		OutFormat.Dim = FIntPoint(width, height);
-		// TODO: Don't hardcode. Get from C-API. ml_media_player api does not provide that right now. Try the ml_media_codec api.
-		OutFormat.FrameRate = 30.0f;
-		OutFormat.FrameRates = TRange<float>(30.0f);
-		OutFormat.TypeName = TEXT("BGRA");
+		UE_LOG(LogMagicLeapMedia, Error, TEXT("MLMediaPlayerGetVideoSize failed with error %d"), Result);
+		return false;
 	}
 
-	return bResult;
+	OutFormat.Dim = FIntPoint(width, height);
+	// TODO: Don't hardcode. Get from C-API. ml_media_player api does not provide that right now. Try the ml_media_codec api.
+	OutFormat.FrameRate = 30.0f;
+	OutFormat.FrameRates = TRange<float>(30.0f);
+	OutFormat.TypeName = TEXT("BGRA");
+	return true;
 }
 
 bool FMagicLeapMediaPlayer::SelectTrack(EMediaTrackType TrackType, int32 TrackIndex)
 {
-	if (TrackInfo.Contains(TrackType))
+	if (TrackInfo.Contains(TrackType) && CurrentState != EMediaState::Preparing)
 	{
-		if (TrackIndex >= 0 && TrackIndex < TrackInfo[TrackType].Num())
+		if (TrackInfo[TrackType].IsValidIndex(TrackIndex))
 		{
-			bool bResult = MLMediaPlayerSelectTrack(MediaPlayerHandle, static_cast<uint32>(TrackInfo[TrackType][TrackIndex]));
-			SelectedTrack[TrackType] = (bResult) ? TrackIndex : SelectedTrack[TrackType];
-			return bResult;
+			MLResult Result = MLMediaPlayerSelectTrack(MediaPlayerHandle, static_cast<uint32>(TrackInfo[TrackType][TrackIndex]));
+			if (Result == MLResult_Ok)
+			{
+				SelectedTrack[TrackType] = TrackIndex;
+				return true;
+			}
+			else
+			{
+				UE_LOG(LogMagicLeapMedia, Error, TEXT("MLMediaPlayerSelectTrack failed with error %d"), Result);
+			}
 		}
 	}
 	return false;
@@ -989,49 +1057,92 @@ FIntPoint FMagicLeapMediaPlayer::GetVideoDimensions() const
 {
 	int32 width = 0;
 	int32 height = 0;
-	MLMediaPlayerGetVideoSize(MediaPlayerHandle, &width, &height);
+	MLResult Result = MLMediaPlayerGetVideoSize(MediaPlayerHandle, &width, &height);
+	if (Result != MLResult_Ok)
+	{
+		UE_LOG(LogMagicLeapMedia, Error, TEXT("MLMediaPlayerGetVideoSize failed with error %d"), Result);
+	}
 	return FIntPoint(width, height);
 }
 
 bool FMagicLeapMediaPlayer::SetRateOne()
 {
-	bool bResult = MLMediaPlayerStart(MediaPlayerHandle);
-	if (bResult)
+	MLResult Result = MLMediaPlayerStart(MediaPlayerHandle);
+	if (Result != MLResult_Ok)
 	{
-		CurrentState = EMediaState::Playing;
-		EventSink.ReceiveMediaEvent(EMediaEvent::PlaybackResumed);
+		UE_LOG(LogMagicLeapMedia, Error, TEXT("MLMediaPlayerStart failed with error %d"), Result);
+		return false;
 	}
-
-	return bResult;
+	
+	CurrentState = EMediaState::Playing;
+	EventSink.ReceiveMediaEvent(EMediaEvent::PlaybackResumed);
+	return true;
 }
 
 bool FMagicLeapMediaPlayer::GetMediaPlayerState(uint16 FlagToPoll) const
 {
-	return FlagToPoll & MLMediaPlayerPollStates(MediaPlayerHandle, FlagToPoll);
+	uint16_t StateFlags = 0;
+	MLResult Result = MLMediaPlayerPollStates(MediaPlayerHandle, FlagToPoll, &StateFlags);
+	if (Result != MLResult_Ok)
+	{
+		UE_LOG(LogMagicLeapMedia, Error, TEXT("MLMediaPlayerPollStates failed with error %d"), Result);
+		return false;
+	}
+
+	return FlagToPoll & StateFlags;
 }
 
 bool FMagicLeapMediaPlayer::RenderThreadIsBufferAvailable(MLHandle InMediaPlayerHandle)
 {
 	ensureMsgf(IsInRenderingThread(), TEXT("RenderThreadIsBufferAvailable called outside of render thread"));
-	return MLMediaPlayerPollingStateFlag_IsBufferAvailable & MLMediaPlayerPollStates(InMediaPlayerHandle, MLMediaPlayerPollingStateFlag_IsBufferAvailable);
+	uint16_t StateFlags = 0;
+	MLResult Result = MLMediaPlayerPollStates(InMediaPlayerHandle, MLMediaPlayerPollingStateFlag_IsBufferAvailable, &StateFlags);
+	if (Result != MLResult_Ok)
+	{
+		UE_LOG(LogMagicLeapMedia, Error, TEXT("MLMediaPlayerPollStates failed with error %d"), Result);
+		return false;
+	}
+
+	return MLMediaPlayerPollingStateFlag_IsBufferAvailable & StateFlags;
 }
 
 bool FMagicLeapMediaPlayer::RenderThreadGetNativeBuffer(const MLHandle InMediaPlayerHandle, MLHandle& NativeBuffer)
 {
 	ensureMsgf(IsInRenderingThread(), TEXT("RenderThreadGetNativeBuffer called outside of render thread"));
-	return MLMediaPlayerAcquireNextAvailableBuffer(InMediaPlayerHandle, &NativeBuffer);
+	MLResult Result = MLMediaPlayerAcquireNextAvailableBuffer(InMediaPlayerHandle, &NativeBuffer);
+	if (Result != MLResult_Ok)
+	{
+		UE_LOG(LogMagicLeapMedia, Error, TEXT("MLMediaPlayerAcquireNextAvailableBuffer failed with error %d"), Result);
+		return false;
+	}
+
+	return true;
 }
 
 bool FMagicLeapMediaPlayer::RenderThreadReleaseNativeBuffer(const MLHandle InMediaPlayerHandle, MLHandle NativeBuffer)
 {
 	ensureMsgf(IsInRenderingThread(), TEXT("RenderThreadReleaseNativeBuffer called outside of render thread"));
-	return MLMediaPlayerReleaseBuffer(InMediaPlayerHandle, NativeBuffer);
+	MLResult Result = MLMediaPlayerReleaseBuffer(InMediaPlayerHandle, NativeBuffer);
+	if (Result != MLResult_Ok)
+	{
+		UE_LOG(LogMagicLeapMedia, Error, TEXT("MLMediaPlayerReleaseBuffer failed with error %d"), Result);
+		return false;
+	}
+
+	return true;
 }
 
 bool FMagicLeapMediaPlayer::RenderThreadGetCurrentPosition(const MLHandle InMediaPlayerHandle, int32& CurrentPosition)
 {
 	ensureMsgf(IsInRenderingThread(), TEXT("RenderThreadGetCurrentPosition called outside of render thread"));
-	return MLMediaPlayerGetCurrentPosition(InMediaPlayerHandle, &CurrentPosition);
+	MLResult Result = MLMediaPlayerGetCurrentPosition(InMediaPlayerHandle, &CurrentPosition);
+	if (Result != MLResult_Ok)
+	{
+		UE_LOG(LogMagicLeapMedia, Error, TEXT("MLMediaPlayerGetCurrentPosition failed with error %d"), Result);
+		return false;
+	}
+
+	return true;
 }
 
 #undef LOCTEXT_NAMESPACE
diff --git a/Engine/Plugins/Lumin/MagicLeapMedia/Source/MagicLeapMedia/Private/MagicLeapMediaPlayer.h b/Engine/Plugins/Lumin/MagicLeapMedia/Source/MagicLeapMedia/Private/MagicLeapMediaPlayer.h
--- a/Engine/Plugins/Lumin/MagicLeapMedia/Source/MagicLeapMedia/Private/MagicLeapMediaPlayer.h
+++ b/Engine/Plugins/Lumin/MagicLeapMedia/Source/MagicLeapMedia/Private/MagicLeapMediaPlayer.h
@@ -103,6 +103,15 @@ protected:
 
 	bool bMediaPrepared;
 
+	/** Frame UV scale and offsets */
+	float UScale = 1.0f;
+	float VScale = 1.0f;
+	float UOffset = 0.0f;
+	float VOffset = 0.0f;
+
+	/** Frame transformation matrix used to determine UV offset and scaling */
+	float FrameTransformationMatrix[16] = {};
+
 	/** Current player state. */
 	EMediaState CurrentState;
 
@@ -145,6 +154,7 @@ protected:
 private:
 	virtual bool SetRateOne();
 	virtual bool GetMediaPlayerState(uint16 FlagToPoll) const;
+	virtual void RegisterExternalTexture(const FGuid& InGuid, FTextureRHIRef& InTextureRHI, FSamplerStateRHIRef& InSamplerStateRHI);
 	virtual bool RenderThreadIsBufferAvailable(MLHandle MediaPlayerHandle);
 	virtual bool RenderThreadGetNativeBuffer(const MLHandle MediaPlayerHandle, MLHandle& NativeBuffer);
 	virtual bool RenderThreadReleaseNativeBuffer(const MLHandle MediaPlayerHandle, MLHandle NativeBuffer);
diff --git a/Engine/Plugins/Lumin/MagicLeapMedia/Source/MagicLeapMedia/Private/MediaWorker.cpp b/Engine/Plugins/Lumin/MagicLeapMedia/Source/MagicLeapMedia/Private/MediaWorker.cpp
--- a/Engine/Plugins/Lumin/MagicLeapMedia/Source/MagicLeapMedia/Private/MediaWorker.cpp
+++ b/Engine/Plugins/Lumin/MagicLeapMedia/Source/MagicLeapMedia/Private/MediaWorker.cpp
@@ -59,19 +59,21 @@ uint32 FMediaWorker::Run()
 {
 	while (StopTaskCounter.GetValue() == 0)
 	{
-		if (!(MLMediaPlayerPollingStateFlag_IsBufferAvailable & MLMediaPlayerPollStates(MediaPlayerHandle, MLMediaPlayerPollingStateFlag_IsBufferAvailable)))
+		uint16_t States = 0;
+		MLResult Result = MLMediaPlayerPollStates(MediaPlayerHandle, MLMediaPlayerPollingStateFlag_IsBufferAvailable, &States);
+		if (Result != MLResult_Ok || !(MLMediaPlayerPollingStateFlag_IsBufferAvailable & States))
 		{
 			continue;
 		}
 
 		int32 currentPositionMilSec = 0;
-		bool bResult = MLMediaPlayerGetCurrentPosition(MediaPlayerHandle, &currentPositionMilSec);
+		Result = MLMediaPlayerGetCurrentPosition(MediaPlayerHandle, &currentPositionMilSec);
 		const FTimespan Time = FTimespan::FromMilliseconds(currentPositionMilSec);
 
 		int32 width = 0;
 		int32 height = 0;
-		bResult = MLMediaPlayerGetVideoSize(MediaPlayerHandle, &width, &height);
-		if (!bResult)
+		Result = MLMediaPlayerGetVideoSize(MediaPlayerHandle, &width, &height);
+		if (Result != MLResult_Ok)
 		{
 			UE_LOG(LogMagicLeapMedia, Error, TEXT("Error getting video dimensions"));
 			continue;
@@ -95,9 +97,9 @@ uint32 FMediaWorker::Run()
 		// MLHandle because Unreal's uint64 is 'unsigned long long *' whereas uint64_t for the C-API is 'unsigned long *'
 		// TODO: Fix the Unreal types for the above comment.
 		MLHandle nativeBuffer = ML_INVALID_HANDLE;
-		bResult = MLMediaPlayerAcquireNextAvailableBuffer(MediaPlayerHandle, &nativeBuffer);
+		Result = MLMediaPlayerAcquireNextAvailableBuffer(MediaPlayerHandle, &nativeBuffer);
 
-		if (!bResult)
+		if (Result != MLResult_Ok)
 		{
 			UE_LOG(LogMagicLeapMedia, Error, TEXT("Error acquiring next video buffer"));
 			continue;
diff --git a/Engine/Plugins/Lumin/MagicLeapScreens/Source/MagicLeapScreens/Classes/ScreensComponent.h b/Engine/Plugins/Lumin/MagicLeapScreens/Source/MagicLeapScreens/Classes/ScreensComponent.h
--- a/Engine/Plugins/Lumin/MagicLeapScreens/Source/MagicLeapScreens/Classes/ScreensComponent.h
+++ b/Engine/Plugins/Lumin/MagicLeapScreens/Source/MagicLeapScreens/Classes/ScreensComponent.h
@@ -28,7 +28,7 @@ struct FScreensWatchHistoryEntry
 
 public:
 	/** Entry Identifier. Must be used to update and delete a given entry. */
-	UPROPERTY(EditAnywhere, BlueprintReadWrite, Category = "Screens|MagicLeap")
+	UPROPERTY(VisibleAnywhere, BlueprintReadOnly, Category = "Screens|MagicLeap")
 	FScreenID ID;
 
 	/** Title of the media for which this entry is created. */
@@ -96,32 +96,30 @@ public:
 	virtual void TickComponent(float DeltaTime, enum ELevelTick TickType, FActorComponentTickFunction *ThisTickFunction) override;
 
 	/** 
-		Attempts to instigate a task on the worker thread to retrieve the watch history. 
-		Subscribe to the GetWatchHistorySuccess & GetWatchHistoryFailure delegates to recevie the result.
-
-		@return True if the worker thread was idle at the time of the call and the task was successfully queued.
+		Queues a task on the worker thread to retrieve the watch history. 
+		Subscribe to the GetWatchHistorySuccess & GetWatchHistoryFailure delegates to receive the result.
 	*/
 	UFUNCTION(BlueprintCallable, Category = "Screens|MagicLeap")
-	bool GetWatchHistoryAsync();
+	void GetWatchHistoryAsync();
 
 	/** 
-		Adds a watch history entry.
+		Queues a task on the worker thread to add a new entry into the watch history.
+		Subscribe to the AddToWatchHistoryResult delegate to receive the result.
 	
 		@param[in] WatchHistoryEntry The entry to be added to the history.
-		@param[out] FScreenID Stores the id of the newly created entry. The same ID should be used to update or delete this entry.
-		@return True if no errors were encountered.
 	*/
 	UFUNCTION(BlueprintCallable, Category = "Screens|MagicLeap")
-	bool AddWatchHistoryEntry(const FScreensWatchHistoryEntry& WatchHistoryEntry, FScreenID& ID);
+	void AddWatchHistoryEntryAsync(const FScreensWatchHistoryEntry WatchHistoryEntry);
 
 	/**
-		Updates a watch history entry.
+		Queues a task on the worker thread to update an entry in the watch history.
+		The ID of the entry to be updated must be valid and already present in the watch history.
+		Subscribe to the UpdateWatchHistoryEntryResult delegate to receive the result.
 	
 		@param[in] WatchHistoryEntry The entry to be updated.
-		@return True if no errors were encountered.
 	*/
 	UFUNCTION(BlueprintCallable, Category = "Screens|MagicLeap")
-	bool UpdateWatchHistoryEntry(const FScreensWatchHistoryEntry& WatchHistoryEntry);
+	void UpdateWatchHistoryEntryAsync(const FScreensWatchHistoryEntry WatchHistoryEntry);
 
 	/**
 		Removes a watch history entry.
@@ -160,19 +158,34 @@ public:
 	/** Delegate used to notify the instigating blueprint of a watch history retrieval failure. */
 	DECLARE_DYNAMIC_MULTICAST_DELEGATE(FScreensGetWatchHistoryFailure);
 
-	/** Fired when the watch history is successfully retrieved on the worker thread. */
-	FScreensGetWatchHistorySuccess& OnScreensGetWatchHistorySuccess();
+	/** 
+		Delegate used to notify the instigating blueprint of the result from a request to add a new watch history entry.
+
+		@param[out] Entry The new entry created if a request via a call to AddToWatchHistoryAsync() succeeds.
+		@param[out] Success True when the call to AddToWatchHistoryAsync() succeeds.
+	*/
+	DECLARE_DYNAMIC_MULTICAST_DELEGATE_TwoParams(FScreensAddToWatchHistoryResult, const FScreensWatchHistoryEntry&, Entry, const bool, Success);
+
+	/**
+		Delegate used to notify the instigating blueprint of the result from a request to update an entry in the watch history
 
-	/** Fired when retrieval of the watch history fails on the worker thread. */
-	FScreensGetWatchHistoryFailure& OnScreensGetWatchHistoryFailure();
+		@param[out] Entry The updated entry resulting from a successful request via a call to UpdateWatchHistoryEntry().
+		@param[out] Success True when the call to UpdateWatchHistoryEntryAsync() succeeds.
+	*/
+	DECLARE_DYNAMIC_MULTICAST_DELEGATE_TwoParams(FScreensUpdateWatchHistoryEntryResult, const FScreensWatchHistoryEntry&, Entry, const bool, Success);
 
 private:
 	class FScreensImpl *Impl;
-	bool bWorkerBusy;
 
 	UPROPERTY(BlueprintAssignable, Category = "Screens|MagicLeap", meta = (AllowPrivateAccess = true))
 	FScreensGetWatchHistorySuccess GetWatchHistorySuccess;
 
 	UPROPERTY(BlueprintAssignable, Category = "Screens|MagicLeap", meta = (AllowPrivateAccess = true))
 	FScreensGetWatchHistoryFailure GetWatchHistoryFailure;
+
+	UPROPERTY(BlueprintAssignable, Category = "Screens|MagicLeap", meta = (AllowPrivateAccess = true))
+	FScreensAddToWatchHistoryResult AddToWatchHistoryResult;
+
+	UPROPERTY(BlueprintAssignable, Category = "Screens|MagicLeap", meta = (AllowPrivateAccess = true))
+	FScreensUpdateWatchHistoryEntryResult UpdateWatchHistoryEntryResult;
 };
diff --git a/Engine/Plugins/Lumin/MagicLeapScreens/Source/MagicLeapScreens/Private/ScreensComponent.cpp b/Engine/Plugins/Lumin/MagicLeapScreens/Source/MagicLeapScreens/Private/ScreensComponent.cpp
--- a/Engine/Plugins/Lumin/MagicLeapScreens/Source/MagicLeapScreens/Private/ScreensComponent.cpp
+++ b/Engine/Plugins/Lumin/MagicLeapScreens/Source/MagicLeapScreens/Private/ScreensComponent.cpp
@@ -9,6 +9,7 @@
 #include "Containers/Array.h"
 #include "RenderUtils.h"
 #include "MagicLeapPluginUtil.h"
+#include "IMagicLeapPlugin.h"
 #include "MagicLeapHMD.h"
 #include "MagicLeapUtils.h"
 #include "IMagicLeapScreensPlugin.h"
@@ -79,6 +80,8 @@ enum ScreensTaskType
 {
 	None,
 	GetHistory,
+	AddToHistory,
+	UpdateEntry,
 };
 
 struct FScreensMessage
@@ -182,9 +185,97 @@ public:
 		{
 		case ScreensTaskType::None: break;
 		case ScreensTaskType::GetHistory: GetWatchHistory(); break;
+		case ScreensTaskType::AddToHistory: AddToHistory();  break;
+		case ScreensTaskType::UpdateEntry: UpdateWatchHistoryEntry(); break;
 		}
 	}
 
+	void AddToHistory()
+	{
+#if WITH_MLSDK
+		check(CurrentMessage.WatchHistory.Num() != 0);
+		FScreensWatchHistoryEntry WatchHistoryEntry = CurrentMessage.WatchHistory[0];
+		FScreensMessage Msg;
+		Msg.Type = ScreensMsgType::Response;
+		Msg.TaskType = ScreensTaskType::AddToHistory;
+
+		MLScreensWatchHistoryEntry Entry;
+		Entry.id = WatchHistoryEntry.ID.ID;
+		Entry.title = TCHAR_TO_UTF8(*WatchHistoryEntry.Title);
+		Entry.subtitle = TCHAR_TO_UTF8(*WatchHistoryEntry.Subtitle);
+		Entry.playback_position_ms = static_cast<uint32>(WatchHistoryEntry.PlaybackPosition.GetTotalMilliseconds());
+		Entry.playback_duration_ms = static_cast<uint32>(WatchHistoryEntry.PlaybackDuration.GetTotalMilliseconds());
+		Entry.custom_data = TCHAR_TO_UTF8(*WatchHistoryEntry.CustomData);
+		MLImage MLThumbnail;
+
+		if (!WatchHistoryEntry.Thumbnail ||
+			!IsSupportedFormat(WatchHistoryEntry.Thumbnail->GetPixelFormat()) ||
+			!UTexture2DToMLImage(*WatchHistoryEntry.Thumbnail, MLThumbnail))
+		{
+			MLThumbnail = DefaultThumbnail;
+		}
+
+		MLResult Result = MLScreensInsertWatchHistoryEntry(&Entry, &MLThumbnail);
+		if (Result != MLResult_Ok)
+		{
+			UE_LOG(LogScreens, Error, TEXT("MLScreensInsertWatchHistoryEntry failed with error %d!"), Result);
+			Msg.WatchHistory.Add(WatchHistoryEntry);
+		}
+		else
+		{
+			FScreensWatchHistoryEntry NewEntry;
+			MLWatchHistoryEntryToUnreal(Entry, NewEntry);
+			Msg.WatchHistory.Add(NewEntry);
+			Msg.bSuccess = true;
+		}
+
+		OutgoingMessages.Enqueue(Msg);
+#endif // WITH_MLSDK
+	}
+
+	void UpdateWatchHistoryEntry()
+	{
+#if WITH_MLSDK
+		check(CurrentMessage.WatchHistory.Num() != 0);
+		FScreensWatchHistoryEntry WatchHistoryEntry = CurrentMessage.WatchHistory[0];
+		FScreensMessage Msg;
+		Msg.Type = ScreensMsgType::Response;
+		Msg.TaskType = ScreensTaskType::UpdateEntry;
+		MLScreensWatchHistoryEntry Entry;
+
+		Entry.id = WatchHistoryEntry.ID.ID;
+		Entry.title = TCHAR_TO_UTF8(*WatchHistoryEntry.Title);
+		Entry.subtitle = TCHAR_TO_UTF8(*WatchHistoryEntry.Subtitle);
+		Entry.playback_position_ms = static_cast<uint32>(WatchHistoryEntry.PlaybackPosition.GetTotalMilliseconds());
+		Entry.playback_duration_ms = static_cast<uint32>(WatchHistoryEntry.PlaybackDuration.GetTotalMilliseconds());
+		Entry.custom_data = TCHAR_TO_UTF8(*WatchHistoryEntry.CustomData);
+		MLImage MLThumbnail;
+
+		if (!WatchHistoryEntry.Thumbnail ||
+			!IsSupportedFormat(WatchHistoryEntry.Thumbnail->GetPixelFormat()) ||
+			!UTexture2DToMLImage(*WatchHistoryEntry.Thumbnail, MLThumbnail))
+		{
+			MLThumbnail = DefaultThumbnail;
+		}
+
+		MLResult Result = MLScreensUpdateWatchHistoryEntry(&Entry, &MLThumbnail);
+		if (Result != MLResult_Ok)
+		{
+			UE_LOG(LogScreens, Error, TEXT("MLScreensUpdateWatchHistoryEntry failed with error %d!"), Result);
+			Msg.WatchHistory.Add(WatchHistoryEntry);
+		}
+		else
+		{
+			FScreensWatchHistoryEntry UpdatedEntry;
+			MLWatchHistoryEntryToUnreal(Entry, UpdatedEntry);
+			Msg.WatchHistory.Add(UpdatedEntry);
+			Msg.bSuccess = true;
+		}
+
+		OutgoingMessages.Enqueue(Msg);
+#endif // WITH_MLSDK
+	}
+
 	void GetWatchHistory()
 	{
 #if WITH_MLSDK
@@ -193,7 +284,8 @@ public:
 		Msg.TaskType = ScreensTaskType::GetHistory;
 		MLScreensWatchHistoryList WatchHistoryList;
 
-		if (MLScreensGetWatchHistoryList(&WatchHistoryList))
+		MLResult Result = MLScreensGetWatchHistoryList(&WatchHistoryList);
+		if (Result == MLResult_Ok)
 		{
 			for (uint32 i = 0; i < WatchHistoryList.count; ++i)
 			{
@@ -201,12 +293,16 @@ public:
 				MLWatchHistoryEntryToUnreal(WatchHistoryList.entries[i], WatchHistoryEntry);
 				Msg.WatchHistory.Add(WatchHistoryEntry);
 			}
-			MLScreensReleaseWatchHistoryList(&WatchHistoryList);
+			Result = MLScreensReleaseWatchHistoryList(&WatchHistoryList);
+			if (Result != MLResult_Ok)
+			{
+				UE_LOG(LogScreens, Error, TEXT("MLScreensReleaseWatchHistoryList failed with error %d!"), Result);
+			}
 			Msg.bSuccess = true;
 		}
 		else
 		{
-			UE_LOG(LogScreens, Error, TEXT("MLScreensGetWatchHistoryList failed!"));
+			UE_LOG(LogScreens, Error, TEXT("MLScreensGetWatchHistoryList failed with error %d!"), Result);
 		}
 
 		OutgoingMessages.Enqueue(Msg);
@@ -243,10 +339,10 @@ public:
 		OutEntry.CustomData = FString(UTF8_TO_TCHAR(InEntry.custom_data));
 
 		MLImage MLThumbnail;
-
-		if (!MLScreensGetWatchHistoryThumbnail(InEntry.id, &MLThumbnail))
+		MLResult Result = MLScreensGetWatchHistoryThumbnail(InEntry.id, &MLThumbnail);
+		if (Result != MLResult_Ok)
 		{
-			UE_LOG(LogScreens, Error, TEXT("Failed to get thumbnail for screen ID %u"), (uint32)InEntry.id);
+			UE_LOG(LogScreens, Error, TEXT("MLScreensGetWatchHistoryThumbnail failed for screen ID %u with error %d!"), (uint32)InEntry.id, Result);
 			MLThumbnail = DefaultThumbnail;
 		}
 
@@ -322,7 +418,6 @@ public:
 
 UScreensComponent::UScreensComponent()
 	: Impl(new FScreensImpl())
-	, bWorkerBusy(false)
 {
 	PrimaryComponentTick.TickGroup = TG_PrePhysics;
 	PrimaryComponentTick.bStartWithTickEnabled = true;
@@ -359,7 +454,6 @@ void UScreensComponent::TickComponent(float DeltaTime, enum ELevelTick TickType,
 	{
 		FScreensMessage Msg;
 		Impl->OutgoingMessages.Dequeue(Msg);
-		bWorkerBusy = false;
 
 		if (Msg.Type == ScreensMsgType::Request)
 		{
@@ -382,100 +476,69 @@ void UScreensComponent::TickComponent(float DeltaTime, enum ELevelTick TickType,
 				}
 			}
 			break;
+			case ScreensTaskType::AddToHistory:
+			{
+				if (Msg.WatchHistory.Num() > 0)
+				{
+					FScreensWatchHistoryEntry Entry = Msg.WatchHistory[0];
+					AddToWatchHistoryResult.Broadcast(Entry, Msg.bSuccess);
+				}
+				else
+				{
+					UE_LOG(LogScreens, Error, TEXT("Unexpected empty watch history in an AddToHistory response from the worker thread"));
+				}
+			}
+			break;
+			case ScreensTaskType::UpdateEntry:
+			{
+				if (Msg.WatchHistory.Num() > 0)
+				{
+					FScreensWatchHistoryEntry Entry = Msg.WatchHistory[0];
+					UpdateWatchHistoryEntryResult.Broadcast(Entry, Msg.bSuccess);
+				}
+				else
+				{
+					UE_LOG(LogScreens, Error, TEXT("Unexpected empty watch history in an UpdateEntry response from the worker thread"));
+				}
+			}
+			break;
 			}
 		}
 	}
 }
 
-UScreensComponent::FScreensGetWatchHistorySuccess& UScreensComponent::OnScreensGetWatchHistorySuccess()
+void UScreensComponent::GetWatchHistoryAsync()
 {
-	return GetWatchHistorySuccess;
-}
-
-UScreensComponent::FScreensGetWatchHistoryFailure& UScreensComponent::OnScreensGetWatchHistoryFailure()
-{
-	return GetWatchHistoryFailure;
-}
-
-bool UScreensComponent::GetWatchHistoryAsync()
-{
-	if (!bWorkerBusy)
-	{
-		bWorkerBusy = true;
-		FScreensMessage Msg;
-		Msg.Type = ScreensMsgType::Request;
-		Msg.TaskType = ScreensTaskType::GetHistory;
-		Impl->ProcessMessage(Msg);
-		return true;
-	}
-
-	UE_LOG(LogScreens, Warning, TEXT("Worker thread is already busy!"));
-	return false;
+	FScreensMessage Msg;
+	Msg.Type = ScreensMsgType::Request;
+	Msg.TaskType = ScreensTaskType::GetHistory;
+	Impl->ProcessMessage(Msg);
 }
 
-bool UScreensComponent::AddWatchHistoryEntry(const FScreensWatchHistoryEntry& WatchHistoryEntry, FScreenID& ID)
+void UScreensComponent::AddWatchHistoryEntryAsync(const FScreensWatchHistoryEntry WatchHistoryEntry)
 {
-#if WITH_MLSDK
-	MLScreensWatchHistoryEntry Entry;
-	Entry.id = WatchHistoryEntry.ID.ID;
-	Entry.title = TCHAR_TO_UTF8(*WatchHistoryEntry.Title);
-	Entry.subtitle = TCHAR_TO_UTF8(*WatchHistoryEntry.Subtitle);
-	Entry.playback_position_ms = static_cast<uint32>(WatchHistoryEntry.PlaybackPosition.GetTotalMilliseconds());
-	Entry.playback_duration_ms = static_cast<uint32>(WatchHistoryEntry.PlaybackDuration.GetTotalMilliseconds());
-	Entry.custom_data = TCHAR_TO_UTF8(*WatchHistoryEntry.CustomData);
-	MLImage MLThumbnail;
-
-	if (!WatchHistoryEntry.Thumbnail ||
-		!Impl->IsSupportedFormat(WatchHistoryEntry.Thumbnail->GetPixelFormat()) ||
-		!Impl->UTexture2DToMLImage(*WatchHistoryEntry.Thumbnail, MLThumbnail))
-	{
-		MLThumbnail = Impl->DefaultThumbnail;
-	}
-
-	if (MLScreensInsertWatchHistoryEntry(&Entry, &MLThumbnail))
-	{
-		ID.ID = Entry.id;
-		return true;
-	}
-
-	UE_LOG(LogScreens, Error, TEXT("MLScreensInsertWatchHistoryEntry failed!"));
-#endif //WITH_MLSDK
-	return false;
+	FScreensMessage Msg;
+	Msg.Type = ScreensMsgType::Request;
+	Msg.TaskType = ScreensTaskType::AddToHistory;
+	Msg.WatchHistory.Add(WatchHistoryEntry);
+	Impl->ProcessMessage(Msg);
 }
 
-bool UScreensComponent::UpdateWatchHistoryEntry(const FScreensWatchHistoryEntry& WatchHistoryEntry)
+void UScreensComponent::UpdateWatchHistoryEntryAsync(const FScreensWatchHistoryEntry WatchHistoryEntry)
 {
-#if WITH_MLSDK
-	MLScreensWatchHistoryEntry Entry;
-	Entry.id = WatchHistoryEntry.ID.ID;
-	Entry.title = TCHAR_TO_UTF8(*WatchHistoryEntry.Title);
-	Entry.subtitle = TCHAR_TO_UTF8(*WatchHistoryEntry.Subtitle);
-	Entry.playback_position_ms = static_cast<uint32>(WatchHistoryEntry.PlaybackPosition.GetTotalMilliseconds());
-	Entry.playback_duration_ms = static_cast<uint32>(WatchHistoryEntry.PlaybackDuration.GetTotalMilliseconds());
-	Entry.custom_data = TCHAR_TO_UTF8(*WatchHistoryEntry.CustomData);
-	MLImage MLThumbnail;
-
-	if (!WatchHistoryEntry.Thumbnail ||
-		!Impl->IsSupportedFormat(WatchHistoryEntry.Thumbnail->GetPixelFormat()) ||
-		!Impl->UTexture2DToMLImage(*WatchHistoryEntry.Thumbnail, MLThumbnail))
-	{
-		MLThumbnail = Impl->DefaultThumbnail;
-	}
-
-	if (MLScreensUpdateWatchHistoryEntry(&Entry, &MLThumbnail))
-	{
-		return true;
-	}
-
-	UE_LOG(LogScreens, Error, TEXT("MLScreensUpdateWatchHistoryEntry failed!"));
-#endif //WITH_MLSDK
-	return false;
+	FScreensMessage Msg;
+	Msg.Type = ScreensMsgType::Request;
+	Msg.TaskType = ScreensTaskType::UpdateEntry;
+	Msg.WatchHistory.Add(WatchHistoryEntry);
+	Impl->ProcessMessage(Msg);
 }
 
 bool UScreensComponent::RemoveWatchHistoryEntry(const FScreenID& ID)
 {
 #if WITH_MLSDK
-	return MLScreensRemoveWatchHistoryEntry(ID.ID);
+	MLResult Result = MLScreensRemoveWatchHistoryEntry(ID.ID);
+	UE_CLOG(Result != MLResult_Ok, LogScreens, Error, TEXT("MLScreensRemoveWatchHistoryEntry failed with error %d for entry with id %d!"), Result, ID.ID);
+	return Result == MLResult_Ok;
 #else
 	return false;
 #endif //WITH_MLSDK
@@ -483,23 +546,39 @@ bool UScreensComponent::RemoveWatchHistoryEntry(const FScreenID& ID)
 
 bool UScreensComponent::ClearWatchHistory()
 {
-	bool bResult = false;
+	bool bSuccess = false;
 #if WITH_MLSDK
 	MLScreensWatchHistoryList WatchHistoryList;
-	bResult = MLScreensGetWatchHistoryList(&WatchHistoryList);
+	MLResult Result = MLScreensGetWatchHistoryList(&WatchHistoryList);
+	bSuccess = true;
 
-	if (bResult)
+	if (Result == MLResult_Ok)
 	{
 		for (uint32 i = 0; i < WatchHistoryList.count; ++i)
 		{
-			MLScreensRemoveWatchHistoryEntry(WatchHistoryList.entries[i].id);
+			Result = MLScreensRemoveWatchHistoryEntry(WatchHistoryList.entries[i].id);
+			if (Result != MLResult_Ok)
+			{
+				bSuccess = false;
+				UE_LOG(LogScreens, Error, TEXT("MLScreensRemoveWatchHistoryEntry failed with error %d for entry %d!"), Result, WatchHistoryList.entries[i].id);
+			}
 		}
 
-		MLScreensReleaseWatchHistoryList(&WatchHistoryList);
+		Result = MLScreensReleaseWatchHistoryList(&WatchHistoryList);
+		if (Result != MLResult_Ok)
+		{
+			bSuccess = false;
+			UE_LOG(LogScreens, Error, TEXT("MLScreensReleaseWatchHistoryList failed with error %d!"), Result);
+		}
+	}
+	else
+	{
+		bSuccess = false;
+		UE_LOG(LogScreens, Error, TEXT("MLScreensGetWatchHistoryList failed with error %d!"), Result);
 	}
 #endif //WITH_MLSDK
 
-	return bResult;
+	return bSuccess;
 }
 
 bool UScreensComponent::GetScreensTransforms(TArray<FScreenTransform>& ScreensTransforms)
@@ -508,6 +587,11 @@ bool UScreensComponent::GetScreensTransforms(TArray<FScreenTransform>& ScreensTr
 #if WITH_MLSDK
 	ScreensTransforms.Empty();
 
+	if (!IMagicLeapPlugin::Get().IsMagicLeapHMDValid())
+	{
+		return false;
+	}
+
 	const FAppFramework& AppFramework = static_cast<FMagicLeapHMD*>(GEngine->XRSystem->GetHMDDevice())->GetAppFrameworkConst();
 	if (!AppFramework.IsInitialized())
 	{
@@ -516,7 +600,8 @@ bool UScreensComponent::GetScreensTransforms(TArray<FScreenTransform>& ScreensTr
 	float WorldToMetersScale = AppFramework.GetWorldToMetersScale();
 
 	MLScreensScreenInfoList ScreensInfoList;
-	bResult = MLScreensGetScreenInfoList(&ScreensInfoList);
+	MLResult Result = MLScreensGetScreenInfoList(&ScreensInfoList);
+	bResult = (Result == MLResult_Ok);
 	if (bResult)
 	{
 		FTransform PoseTransform = UHeadMountedDisplayFunctionLibrary::GetTrackingToWorldTransform(this);
@@ -537,15 +622,23 @@ bool UScreensComponent::GetScreensTransforms(TArray<FScreenTransform>& ScreensTr
 
 			ScreenTransform.ScreenPosition = screen.GetLocation();
 			ScreenTransform.ScreenOrientation = screen.Rotator();
-			ScreenTransform.ScreenDimensions = MagicLeap::ToFVector(Entry.scale, WorldToMetersScale);
+			ScreenTransform.ScreenDimensions = MagicLeap::ToFVector(Entry.dimensions, WorldToMetersScale);
 			ScreenTransform.ScreenDimensions.X = FMath::Abs<float>(ScreenTransform.ScreenDimensions.X);
 			ScreenTransform.ScreenDimensions.Y = FMath::Abs<float>(ScreenTransform.ScreenDimensions.Y);
 			ScreenTransform.ScreenDimensions.Z = FMath::Abs<float>(ScreenTransform.ScreenDimensions.Z);
 
 			ScreensTransforms.Add(ScreenTransform);
 		}
-		MLScreensReleaseScreenInfoList(&ScreensInfoList);
+		Result = MLScreensReleaseScreenInfoList(&ScreensInfoList);
+		if (Result != MLResult_Ok)
+		{
+			UE_LOG(LogScreens, Error, TEXT("MLScreensReleaseScreenInfoList failed with error %d!"), Result);
+		}
 	}
+	else
+	{
+		UE_LOG(LogScreens, Error, TEXT("MLScreensGetScreenInfoList failed with error %d!"), Result);
+	}	
 #endif //WITH_MLSDK
 	return bResult;
 }
\ No newline at end of file
diff --git a/Engine/Plugins/Runtime/AudioCapture/Source/AudioCapture/Private/AudioCaptureComponent.cpp b/Engine/Plugins/Runtime/AudioCapture/Source/AudioCapture/Private/AudioCaptureComponent.cpp
--- a/Engine/Plugins/Runtime/AudioCapture/Source/AudioCapture/Private/AudioCaptureComponent.cpp
+++ b/Engine/Plugins/Runtime/AudioCapture/Source/AudioCapture/Private/AudioCaptureComponent.cpp
@@ -91,6 +91,7 @@ void UAudioCaptureComponent::OnEndGenerate()
 		check(CaptureSynth.IsStreamOpen());
 		CaptureSynth.AbortCapturing();
 		check(!CaptureSynth.IsStreamOpen());
+		bIsStreamOpen = false;
 
 		bIsReadyForForFinishDestroy = true;
 	}
diff --git a/Engine/Plugins/Runtime/AudioCapture/Source/AudioCapture/Private/AudioCaptureMLAudio.cpp b/Engine/Plugins/Runtime/AudioCapture/Source/AudioCapture/Private/AudioCaptureMLAudio.cpp
--- a/Engine/Plugins/Runtime/AudioCapture/Source/AudioCapture/Private/AudioCaptureMLAudio.cpp
+++ b/Engine/Plugins/Runtime/AudioCapture/Source/AudioCapture/Private/AudioCaptureMLAudio.cpp
@@ -15,20 +15,25 @@ FAudioCaptureImpl::FAudioCaptureImpl()
 {
 }
 
-static void OnAudioCaptureCallback(void* CallbackContext)
+static void OnAudioCaptureCallback(MLHandle Handle, void* CallbackContext)
 {
+	(void)Handle;
 	FAudioCaptureImpl* AudioCapture = (FAudioCaptureImpl*)CallbackContext;
 	check(MLHandleIsValid(AudioCapture->InputDeviceHandle));
 	MLAudioBuffer OutputBuffer;
-	MLAudioError Result;
-	if (MLAudioGetInputStreamBuffer(AudioCapture->InputDeviceHandle, &OutputBuffer, &Result))
+	MLResult Result = MLAudioGetInputStreamBuffer(AudioCapture->InputDeviceHandle, &OutputBuffer);
+	if (Result == MLResult_Ok)
 	{
 		AudioCapture->OnAudioCapture(OutputBuffer.ptr, OutputBuffer.size / sizeof(int16), 0.0, false);
-		MLAudioReleaseInputStreamBuffer(AudioCapture->InputDeviceHandle, &Result);
+		Result = MLAudioReleaseInputStreamBuffer(AudioCapture->InputDeviceHandle);
+		if (Result != MLResult_Ok)
+		{
+			UE_LOG(LogAudioCapture, Warning, TEXT("MLAudioReleaseInputStreamBuffer failed with error %d"), Result);
+		}
 	}
 	else
 	{
-		UE_LOG(LogAudioCapture, Warning, TEXT("CALLBACK ERROR %u"), Result);
+		UE_LOG(LogAudioCapture, Warning, TEXT("MLAudioGetInputStreamBuffer failed with error %d"), Result);
 	}
 }
 
@@ -60,11 +65,10 @@ bool FAudioCaptureImpl::GetDefaultCaptureDeviceInfo(FCaptureDeviceInfo& OutInfo)
 	MLAudioBufferFormat DefaultBufferFormat;
 	uint32 BufferSize = 0;
 	uint32 MinBufferSize = 0;
-	MLAudioError Result;
 	uint32 UnsignedSampleRate = SampleRate;
-	MLAudioGetInputStreamDefaults(ChannelCount, UnsignedSampleRate, &DefaultBufferFormat, &BufferSize, &MinBufferSize, &Result);
 
-	if (Result == MLAudioError_Success)
+	MLResult Result = MLAudioGetInputStreamDefaults(ChannelCount, UnsignedSampleRate, &DefaultBufferFormat, &BufferSize, &MinBufferSize);
+	if (Result == MLResult_Ok)
 	{
 		OutInfo.DeviceName = TEXT("MLAudio Microphones");
 		OutInfo.InputChannels = ChannelCount;
@@ -73,7 +77,7 @@ bool FAudioCaptureImpl::GetDefaultCaptureDeviceInfo(FCaptureDeviceInfo& OutInfo)
 	}
 	else
 	{
-		UE_LOG(LogAudioCapture, Error, TEXT("Unable to retrieve settings from MLAudio error: %u! Sample rate %u"), Result, UnsignedSampleRate);
+		UE_LOG(LogAudioCapture, Error, TEXT("Unable to retrieve settings from MLAudio error: %d! Sample rate %u"), Result, UnsignedSampleRate);
 	}
 	
 
@@ -97,25 +101,25 @@ bool FAudioCaptureImpl::OpenDefaultCaptureStream(const FAudioCaptureStreamParam&
 	// set up variables to be populated by ML Audio
 	uint32 ChannelCount = NumChannels;
 	MLAudioBufferFormat DefaultBufferFormat;
-	uint32 BufferSize = 0;
+	uint32 RecommendedBufferSize = 0;
 	uint32 MinBufferSize = 0;
 	uint32 UnsignedSampleRate = SampleRate;
-	MLAudioError Result;
-
-	MLAudioGetInputStreamDefaults(ChannelCount, UnsignedSampleRate, &DefaultBufferFormat, &BufferSize, &MinBufferSize, &Result);
-	
 	uint32 NumFrames = StreamParams.NumFramesDesired;
-	uint32 RequestedBufferSize = StreamParams.NumFramesDesired * NumChannels * sizeof(int16);
+	uint32 BufferSize = StreamParams.NumFramesDesired * NumChannels * sizeof(int16);
 
-	if (RequestedBufferSize < MinBufferSize)
+	MLResult Result = MLAudioGetInputStreamDefaults(ChannelCount, UnsignedSampleRate, &DefaultBufferFormat, &RecommendedBufferSize, &MinBufferSize);
+	if (Result != MLResult_Ok)
 	{
-		UE_LOG(LogAudioCapture, Warning, TEXT("Requested buffer size of %u is smaller than the minimum buffer size, reverting to a buffer size of %u"), RequestedBufferSize, BufferSize);
+		UE_LOG(LogAudioCapture, Error, TEXT("MLAudioGetInputStreamDefaults failed with error %d"), Result);
+		return false;
 	}
-	else
+
+	if (BufferSize < RecommendedBufferSize)
 	{
-		BufferSize = StreamParams.NumFramesDesired * NumChannels;
-		UE_LOG(LogAudioCapture, Display, TEXT("Using buffer size of %u"), StreamParams.NumFramesDesired * NumChannels);
+		UE_LOG(LogAudioCapture, Warning, TEXT("Requested buffer size of %u is smaller than the recommended buffer size, reverting to a buffer size of %u"), BufferSize, RecommendedBufferSize);
+		BufferSize = RecommendedBufferSize;
 	}
+	UE_LOG(LogAudioCapture, Display, TEXT("Using buffer size of %u"), BufferSize);
 
 	DefaultBufferFormat.bits_per_sample = 16;
 	DefaultBufferFormat.sample_format = MLAudioSampleFormat_Int;
@@ -125,14 +129,14 @@ bool FAudioCaptureImpl::OpenDefaultCaptureStream(const FAudioCaptureStreamParam&
 	Callback = StreamParams.Callback;
 
 	// Open up new audio stream
-	MLAudioCreateInputFromVoiceComm(&DefaultBufferFormat, BufferSize, &OnAudioCaptureCallback, this, &InputDeviceHandle, &Result);
-	if (Result != MLAudioError_Success)
+	Result = MLAudioCreateInputFromVoiceComm(&DefaultBufferFormat, BufferSize, &OnAudioCaptureCallback, this, &InputDeviceHandle);
+	if (Result != MLResult_Ok)
 	{
-		UE_LOG(LogAudioCapture, Warning, TEXT("MLAudioCreateInputFromVoiceComm failed with code %u"), Result);
+		UE_LOG(LogAudioCapture, Warning, TEXT("MLAudioCreateInputFromVoiceComm failed with code %d"), Result);
 	}
 	else if (InputDeviceHandle == ML_INVALID_HANDLE)
 	{
-		UE_LOG(LogAudioCapture, Warning, TEXT("MLAudioCreateInputFromVoiceComm failed to generate an input device handle"), Result);
+		UE_LOG(LogAudioCapture, Warning, TEXT("MLAudioCreateInputFromVoiceComm failed to generate an input device handle."));
 	}
 	return true;
 }
@@ -140,10 +144,10 @@ bool FAudioCaptureImpl::OpenDefaultCaptureStream(const FAudioCaptureStreamParam&
 bool FAudioCaptureImpl::CloseStream()
 {
 	UE_LOG(LogAudioCapture, Warning, TEXT("CLOSING STREAM"));
-	MLAudioError Result;
-	if (!MLAudioDestroyInput(InputDeviceHandle, &Result))
+	MLResult Result = MLAudioDestroyInput(InputDeviceHandle);
+	if (Result != MLResult_Ok)
 	{
-		UE_LOG(LogAudioCapture, Warning, TEXT("MLAudioDestroyInput failed with code %u"), Result);
+		UE_LOG(LogAudioCapture, Warning, TEXT("MLAudioDestroyInput failed with code %d"), Result);
 		return false;
 	}
 
@@ -154,10 +158,10 @@ bool FAudioCaptureImpl::CloseStream()
 bool FAudioCaptureImpl::StartStream()
 {
 	UE_LOG(LogAudioCapture, Warning, TEXT("STARTING STREAM"));
-	MLAudioError Result;
-	if (!MLAudioStartInput(InputDeviceHandle, &Result))
+	MLResult Result = MLAudioStartInput(InputDeviceHandle);
+	if (Result != MLResult_Ok)
 	{
-		UE_LOG(LogAudioCapture, Warning, TEXT("MLAudioStartInput failed with code %u"), Result);
+		UE_LOG(LogAudioCapture, Warning, TEXT("MLAudioStartInput failed with code %d"), Result);
 		return false;
 	}
 
@@ -168,10 +172,10 @@ bool FAudioCaptureImpl::StartStream()
 bool FAudioCaptureImpl::StopStream()
 {
 	UE_LOG(LogAudioCapture, Warning, TEXT("STOPPING STREAM"));
-	MLAudioError Result;
-	if (!MLAudioStopInput(InputDeviceHandle, &Result))
+	MLResult Result = MLAudioStopInput(InputDeviceHandle);
+	if (Result != MLResult_Ok)
 	{
-		UE_LOG(LogAudioCapture, Warning, TEXT("MLAudioStopInput failed with code %u"), Result);
+		UE_LOG(LogAudioCapture, Warning, TEXT("MLAudioStopInput failed with code %d"), Result);
 		return false;
 	}
 
diff --git a/Engine/Source/Developer/Android/AndroidDeviceDetection/Private/AndroidDeviceDetectionModule.cpp b/Engine/Source/Developer/Android/AndroidDeviceDetection/Private/AndroidDeviceDetectionModule.cpp
--- a/Engine/Source/Developer/Android/AndroidDeviceDetection/Private/AndroidDeviceDetectionModule.cpp
+++ b/Engine/Source/Developer/Android/AndroidDeviceDetection/Private/AndroidDeviceDetectionModule.cpp
@@ -84,11 +84,12 @@ public:
 		return 0;
 	}
 
-	void UpdateADBPath(FString &InADBPath, FString& InGetPropCommand, bool InbGetExtensionsViaSurfaceFlinger)
+	void UpdateADBPath(FString &InADBPath, FString& InGetPropCommand, bool InbGetExtensionsViaSurfaceFlinger, bool InbForLumin)
 	{
 		ADBPath = InADBPath;
 		GetPropCommand = InGetPropCommand;
 		bGetExtensionsViaSurfaceFlinger = InbGetExtensionsViaSurfaceFlinger;
+		bForLumin = InbForLumin;
 
 		HasADBPath = !ADBPath.IsEmpty();
 		// Force a check next time we go around otherwise it can take over 10sec to find devices
@@ -175,6 +176,22 @@ private:
 			const FString DeviceState = DeviceString.Mid(TabIndex + 1).TrimStart();
 
 			NewDeviceInfo.bAuthorizedDevice = DeviceState != TEXT("unauthorized");
+			if (bForLumin)
+			{
+				// 'mldb oobestaus' is deprecated. 'mldb ps' gives us similar functionality for checking device readiness to some extent.
+				const FString OobeCommand = FString::Printf(TEXT("-s %s ps"), *NewDeviceInfo.SerialNumber);
+				FString OobeStatus;
+				NewDeviceInfo.bAuthorizedDevice = ExecuteAdbCommand(*OobeCommand, &OobeStatus, nullptr);
+				if (DeviceMap.Contains(NewDeviceInfo.SerialNumber))
+				{
+					if (DeviceMap[NewDeviceInfo.SerialNumber].bAuthorizedDevice != NewDeviceInfo.bAuthorizedDevice)
+					{
+						// if this device is already in the connected list but authorization has changed, remove it.
+						// it will be added in the next query which will allow UI to refresh properly.
+						continue;
+					}
+				}
+			}
 
 			// add it to our list of currently connected devices
 			CurrentlyConnectedDevices.Add(NewDeviceInfo.SerialNumber);
@@ -182,12 +199,12 @@ private:
 			// move on to next device if this one is already a known device that has either already been authorized or the authorization
 			// status has not changed
 			if (DeviceMap.Contains(NewDeviceInfo.SerialNumber) && 
-				(DeviceMap[NewDeviceInfo.SerialNumber].bAuthorizedDevice || !NewDeviceInfo.bAuthorizedDevice))
-			{					
+				(DeviceMap[NewDeviceInfo.SerialNumber].bAuthorizedDevice == NewDeviceInfo.bAuthorizedDevice))
+			{
 				continue;
 			}
 
-			if (!NewDeviceInfo.bAuthorizedDevice)
+			if (!NewDeviceInfo.bAuthorizedDevice && !bForLumin)
 			{
 				//note: AndroidTargetDevice::GetName() does not fetch this value, do not rely on this
 				NewDeviceInfo.DeviceName = TEXT("Unauthorized - enable USB debugging");
@@ -356,6 +373,7 @@ private:
 	FString ADBPath;
 	FString GetPropCommand;
 	bool bGetExtensionsViaSurfaceFlinger;
+	bool bForLumin;
 
 	// > 0 if we've been asked to abort work in progress at the next opportunity
 	FThreadSafeCounter StopTaskCounter;
@@ -392,12 +410,13 @@ public:
 		}
 	}
 
-	virtual void Initialize(const TCHAR* InSDKDirectoryEnvVar, const TCHAR* InSDKRelativeExePath, const TCHAR* InGetPropCommand, bool InbGetExtensionsViaSurfaceFlinger) override
+	virtual void Initialize(const TCHAR* InSDKDirectoryEnvVar, const TCHAR* InSDKRelativeExePath, const TCHAR* InGetPropCommand, bool InbGetExtensionsViaSurfaceFlinger, bool InbForLumin = false) override
 	{
 		SDKDirEnvVar = InSDKDirectoryEnvVar;
 		SDKRelativeExePath = InSDKRelativeExePath;
 		GetPropCommand = InGetPropCommand;
 		bGetExtensionsViaSurfaceFlinger = InbGetExtensionsViaSurfaceFlinger;
+		bForLumin = InbForLumin;
 		UpdateADBPath();
 	}
 
@@ -473,7 +492,7 @@ public:
 				ADBPath.Empty();
 			}
 		}
-		DetectionThreadRunnable->UpdateADBPath(ADBPath, GetPropCommand, bGetExtensionsViaSurfaceFlinger);
+		DetectionThreadRunnable->UpdateADBPath(ADBPath, GetPropCommand, bGetExtensionsViaSurfaceFlinger, bForLumin);
 	}
 
 private:
@@ -485,6 +504,7 @@ private:
 	FString SDKRelativeExePath;
 	FString GetPropCommand;
 	bool bGetExtensionsViaSurfaceFlinger;
+	bool bForLumin;
 
 	FRunnableThread* DetectionThread;
 	FAndroidDeviceDetectionRunnable* DetectionThreadRunnable;
diff --git a/Engine/Source/Developer/Android/AndroidDeviceDetection/Public/Interfaces/IAndroidDeviceDetection.h b/Engine/Source/Developer/Android/AndroidDeviceDetection/Public/Interfaces/IAndroidDeviceDetection.h
--- a/Engine/Source/Developer/Android/AndroidDeviceDetection/Public/Interfaces/IAndroidDeviceDetection.h
+++ b/Engine/Source/Developer/Android/AndroidDeviceDetection/Public/Interfaces/IAndroidDeviceDetection.h
@@ -56,7 +56,7 @@ struct FAndroidDeviceInfo
 class IAndroidDeviceDetection
 {
 public:
-	virtual void Initialize(const TCHAR* SDKDirectoryEnvVar, const TCHAR* SDKRelativeExePath, const TCHAR* GetPropCommand, bool bGetExtensionsViaSurfaceFlinger) = 0;
+	virtual void Initialize(const TCHAR* SDKDirectoryEnvVar, const TCHAR* SDKRelativeExePath, const TCHAR* GetPropCommand, bool bGetExtensionsViaSurfaceFlinger, bool bForLumin = false) = 0;
 	virtual const TMap<FString,FAndroidDeviceInfo>& GetDeviceMap() = 0;
 	virtual FCriticalSection* GetDeviceMapLock() = 0;
 	virtual void UpdateADBPath() = 0;
diff --git a/Engine/Source/Developer/Android/AndroidTargetPlatform/Private/AndroidTargetPlatform.cpp b/Engine/Source/Developer/Android/AndroidTargetPlatform/Private/AndroidTargetPlatform.cpp
--- a/Engine/Source/Developer/Android/AndroidTargetPlatform/Private/AndroidTargetPlatform.cpp
+++ b/Engine/Source/Developer/Android/AndroidTargetPlatform/Private/AndroidTargetPlatform.cpp
@@ -208,6 +208,11 @@ FAndroidTargetPlatform::~FAndroidTargetPlatform()
 	 FTicker::GetCoreTicker().RemoveTicker(TickDelegateHandle);
 }
 
+FAndroidTargetDevicePtr FAndroidTargetPlatform::CreateTargetDevice(const ITargetPlatform& InTargetPlatform, const FString& InSerialNumber, const FString& InAndroidVariant) const
+{
+	return MakeShareable(new FAndroidTargetDevice(InTargetPlatform, InSerialNumber, InAndroidVariant));
+}
+
 bool FAndroidTargetPlatform::SupportsES2() const
 {
 	// default to support ES2
@@ -778,7 +783,7 @@ bool FAndroidTargetPlatform::HandleTicker( float DeltaTime )
 			// create target device
 			FAndroidTargetDevicePtr& Device = Devices.Add(DeviceInfo.SerialNumber);
 
-			Device = MakeShareable(new FAndroidTargetDevice(*this, DeviceInfo.SerialNumber, GetAndroidVariantName()));
+			Device = CreateTargetDevice(*this, DeviceInfo.SerialNumber, GetAndroidVariantName());
 
 			Device->SetConnected(true);
 			Device->SetModel(DeviceInfo.Model);
diff --git a/Engine/Source/Developer/Android/AndroidTargetPlatform/Private/AndroidTargetPlatform.h b/Engine/Source/Developer/Android/AndroidTargetPlatform/Private/AndroidTargetPlatform.h
--- a/Engine/Source/Developer/Android/AndroidTargetPlatform/Private/AndroidTargetPlatform.h
+++ b/Engine/Source/Developer/Android/AndroidTargetPlatform/Private/AndroidTargetPlatform.h
@@ -264,6 +264,8 @@ protected:
 		return true;
 	}
 
+	virtual FAndroidTargetDevicePtr CreateTargetDevice(const ITargetPlatform& InTargetPlatform, const FString& InSerialNumber, const FString& InAndroidVariant) const;
+
 	// query for rene3ring mode support
 	bool SupportsES2() const;
 	bool SupportsES31() const;
diff --git a/Engine/Source/Developer/LauncherServices/Private/Profiles/LauncherProfile.h b/Engine/Source/Developer/LauncherServices/Private/Profiles/LauncherProfile.h
--- a/Engine/Source/Developer/LauncherServices/Private/Profiles/LauncherProfile.h
+++ b/Engine/Source/Developer/LauncherServices/Private/Profiles/LauncherProfile.h
@@ -2476,6 +2476,7 @@ protected:
 		}
 
 		ValidatePlatformSDKs();
+		ValidateDeviceStatus();
 	}
 	
 	void ValidatePlatformSDKs(void)
@@ -2539,7 +2540,33 @@ protected:
 			}
 		}
 	}
-	
+
+	void ValidateDeviceStatus(void)
+	{
+		ValidationErrors.Remove(ELauncherProfileValidationErrors::LaunchDeviceIsUnauthorized);
+
+		if (DeployedDeviceGroup.IsValid())
+		{
+			const TArray<FString>& Devices = DeployedDeviceGroup->GetDeviceIDs();
+			for (auto DeviceId : Devices)
+			{
+				ITargetDeviceServicesModule* TargetDeviceServicesModule = static_cast<ITargetDeviceServicesModule*>(FModuleManager::Get().LoadModule(TEXT("TargetDeviceServices")));
+				if (TargetDeviceServicesModule)
+				{
+					TSharedPtr<ITargetDeviceProxy> DeviceProxy = TargetDeviceServicesModule->GetDeviceProxyManager()->FindProxyDeviceForTargetDevice(DeviceId);
+					if (DeviceProxy.IsValid())
+					{
+						if (!DeviceProxy->IsAuthorized())
+						{
+							ValidationErrors.Add(ELauncherProfileValidationErrors::LaunchDeviceIsUnauthorized);
+							return;
+						}
+					}
+				}
+			}
+		}
+	}
+
 	void OnLauncherDeviceGroupDeviceAdded(const ILauncherDeviceGroupRef& DeviceGroup, const FString& DeviceId)
 	{
 		if( DeviceGroup == DeployedDeviceGroup )
diff --git a/Engine/Source/Developer/LauncherServices/Public/ILauncherProfile.h b/Engine/Source/Developer/LauncherServices/Public/ILauncherProfile.h
--- a/Engine/Source/Developer/LauncherServices/Public/ILauncherProfile.h
+++ b/Engine/Source/Developer/LauncherServices/Public/ILauncherProfile.h
@@ -198,6 +198,9 @@ namespace ELauncherProfileValidationErrors
 
 		/** The archive step requires a directory to be specified */
 		NoArchiveDirectorySpecified,
+
+		/** Device is unauthorized or is locked */
+		LaunchDeviceIsUnauthorized,
 	};
 }
 
diff --git a/Engine/Source/Developer/Lumin/LuminTargetPlatform/Private/LuminTargetDevice.h b/Engine/Source/Developer/Lumin/LuminTargetPlatform/Private/LuminTargetDevice.h
new file mode 100644
--- /dev/null
+++ b/Engine/Source/Developer/Lumin/LuminTargetPlatform/Private/LuminTargetDevice.h
@@ -0,0 +1,16 @@
+// Copyright 1998-2018 Epic Games, Inc. All Rights Reserved.
+
+#include "AndroidTargetDevice.h"
+
+class FLuminTargetDevice : public FAndroidTargetDevice
+{
+public:
+	FLuminTargetDevice(const ITargetPlatform& InTargetPlatform, const FString& InSerialNumber, const FString& InAndroidVariant)
+	: FAndroidTargetDevice(InTargetPlatform, InSerialNumber, InAndroidVariant)
+	{}
+
+	virtual bool IsPlatformAggregated() const override
+	{
+		return false;
+	}
+};
diff --git a/Engine/Source/Developer/Lumin/LuminTargetPlatform/Private/LuminTargetPlatform.cpp b/Engine/Source/Developer/Lumin/LuminTargetPlatform/Private/LuminTargetPlatform.cpp
--- a/Engine/Source/Developer/Lumin/LuminTargetPlatform/Private/LuminTargetPlatform.cpp
+++ b/Engine/Source/Developer/Lumin/LuminTargetPlatform/Private/LuminTargetPlatform.cpp
@@ -8,6 +8,7 @@
 #include "UObject/NameTypes.h"
 #include "Logging/LogMacros.h"
 #include "Stats/Stats.h"
+#include "LuminTargetDevice.h"
 #include "Modules/ModuleManager.h"
 
 /*=============================================================================
@@ -116,6 +117,11 @@ bool FLuminTargetPlatform::SupportsFeature(ETargetPlatformFeatures Feature) cons
 	return TTargetPlatformBase<FAndroidPlatformProperties>::SupportsFeature(Feature);
 }
 
+FAndroidTargetDevicePtr FLuminTargetPlatform::CreateTargetDevice(const ITargetPlatform& InTargetPlatform, const FString& InSerialNumber, const FString& InAndroidVariant) const
+{
+	return MakeShareable(new FLuminTargetDevice(InTargetPlatform, InSerialNumber, InAndroidVariant));
+}
+
 void FLuminTargetPlatform::InitializeDeviceDetection()
 {
 	DeviceDetection = FModuleManager::LoadModuleChecked<IAndroidDeviceDetectionModule>("AndroidDeviceDetection").GetAndroidDeviceDetection(TEXT("Lumin"));
@@ -125,7 +131,7 @@ void FLuminTargetPlatform::InitializeDeviceDetection()
 #else
 	TEXT("tools/mldb/mldb"),
 #endif
-	TEXT("getprop"), false);
+	TEXT("getprop"), /*bGetExtensionsViaSurfaceFlinger*/ false, /*bForLumin*/ true);
 
 }
 
diff --git a/Engine/Source/Developer/Lumin/LuminTargetPlatform/Private/LuminTargetPlatform.h b/Engine/Source/Developer/Lumin/LuminTargetPlatform/Private/LuminTargetPlatform.h
--- a/Engine/Source/Developer/Lumin/LuminTargetPlatform/Private/LuminTargetPlatform.h
+++ b/Engine/Source/Developer/Lumin/LuminTargetPlatform/Private/LuminTargetPlatform.h
@@ -94,6 +94,9 @@ public:
 
 	virtual bool SupportsFeature(ETargetPlatformFeatures Feature) const override;
 
+protected:
+	virtual FAndroidTargetDevicePtr CreateTargetDevice(const ITargetPlatform& InTargetPlatform, const FString& InSerialNumber, const FString& InAndroidVariant) const override;
+
 protected:
 	// Holds the Engine INI settings, for quick use.
 	FConfigFile EngineSettings;
diff --git a/Engine/Source/Developer/ProjectLauncher/Private/Widgets/Shared/SProjectLauncherValidation.h b/Engine/Source/Developer/ProjectLauncher/Private/Widgets/Shared/SProjectLauncherValidation.h
--- a/Engine/Source/Developer/ProjectLauncher/Private/Widgets/Shared/SProjectLauncherValidation.h
+++ b/Engine/Source/Developer/ProjectLauncher/Private/Widgets/Shared/SProjectLauncherValidation.h
@@ -156,6 +156,12 @@ public:
 					MakeValidationMessage(TEXT("Icons.Error"), LOCTEXT("UnversionedAndIncrimental", "Unversioned build cannot be incremental.").ToString(), ELauncherProfileValidationErrors::UnversionedAndIncrimental)
 				]
 
+			+ SVerticalBox::Slot()
+				.AutoHeight()
+				[
+					MakeValidationMessage(TEXT("Icons.Error"), LOCTEXT("LaunchDeviceIsUnauthorized", "Device is unauthorized or locked.").ToString(), ELauncherProfileValidationErrors::LaunchDeviceIsUnauthorized)
+				]
+
 			+ SVerticalBox::Slot()
 				.AutoHeight()
 				[
diff --git a/Engine/Source/Developer/TargetDeviceServices/Private/Proxies/TargetDeviceProxy.cpp b/Engine/Source/Developer/TargetDeviceServices/Private/Proxies/TargetDeviceProxy.cpp
--- a/Engine/Source/Developer/TargetDeviceServices/Private/Proxies/TargetDeviceProxy.cpp
+++ b/Engine/Source/Developer/TargetDeviceServices/Private/Proxies/TargetDeviceProxy.cpp
@@ -14,6 +14,7 @@
 
 FTargetDeviceProxy::FTargetDeviceProxy(const FString& InName)
 	: Connected(false)
+	, Authorized(false)
 	, LastUpdateTime(0)
 	, Name(InName)
 	, SupportsMultiLaunch(false)
@@ -29,6 +30,7 @@ FTargetDeviceProxy::FTargetDeviceProxy(const FString& InName)
 
 FTargetDeviceProxy::FTargetDeviceProxy(const FString& InName, const FTargetDeviceServicePong& Message, const TSharedRef<IMessageContext, ESPMode::ThreadSafe>& Context, bool InIsAggregated)
 	: Connected(false)
+	, Authorized(false)
 	, Name(InName)
 	, SupportsMultiLaunch(false)
 	, SupportsPowerOff(false)
@@ -54,6 +56,7 @@ void FTargetDeviceProxy::UpdateFromMessage( const FTargetDeviceServicePong& Mess
 	MessageAddress = Context->GetSender();
 
 	Connected = Message.Connected;
+	Authorized = Message.Authorized;
 	HostName = Message.HostName;
 	HostUser = Message.HostUser;
 	Make = Message.Make;
diff --git a/Engine/Source/Developer/TargetDeviceServices/Private/Proxies/TargetDeviceProxy.h b/Engine/Source/Developer/TargetDeviceServices/Private/Proxies/TargetDeviceProxy.h
--- a/Engine/Source/Developer/TargetDeviceServices/Private/Proxies/TargetDeviceProxy.h
+++ b/Engine/Source/Developer/TargetDeviceServices/Private/Proxies/TargetDeviceProxy.h
@@ -156,6 +156,11 @@ public:
 		return Connected;
 	}
 
+	virtual bool IsAuthorized() const override
+	{
+		return Authorized;
+	}
+
 	virtual bool IsShared() const override
 	{
 		return Shared;
@@ -210,6 +215,9 @@ private:
 	/** Holds a flag indicating whether the device is connected. */
 	bool Connected;
 
+	/** Holds a flag indicating whether the device is authorized */
+	bool Authorized;
+
 	/** Holds the name of the computer that hosts the device. */
 	FString HostName;
 
diff --git a/Engine/Source/Developer/TargetDeviceServices/Private/Services/TargetDeviceService.cpp b/Engine/Source/Developer/TargetDeviceServices/Private/Services/TargetDeviceService.cpp
--- a/Engine/Source/Developer/TargetDeviceServices/Private/Services/TargetDeviceService.cpp
+++ b/Engine/Source/Developer/TargetDeviceServices/Private/Services/TargetDeviceService.cpp
@@ -488,6 +488,7 @@ void FTargetDeviceService::HandlePingMessage(const FTargetDeviceServicePing& InM
 		Message->HostName = FPlatformProcess::ComputerName();
 		Message->HostUser = FPlatformProcess::UserName(false);
 		Message->Connected = DefaultDevice->IsConnected();
+		Message->Authorized = DefaultDevice->IsAuthorized();
 		Message->Make = TEXT("@todo");
 		Message->Model = TEXT("@todo");
 		DefaultDevice->GetUserCredentials(Message->DeviceUser, Message->DeviceUserPassword);
diff --git a/Engine/Source/Developer/TargetDeviceServices/Public/ITargetDeviceProxy.h b/Engine/Source/Developer/TargetDeviceServices/Public/ITargetDeviceProxy.h
--- a/Engine/Source/Developer/TargetDeviceServices/Public/ITargetDeviceProxy.h
+++ b/Engine/Source/Developer/TargetDeviceServices/Public/ITargetDeviceProxy.h
@@ -254,6 +254,13 @@ public:
 	 */
 	virtual bool IsConnected() const = 0;
 
+	/**
+	 * Checks whether the device is authorized.
+	 *
+	 * @return true if the device is authorized, false otherwise.
+	 */
+	virtual bool IsAuthorized() const = 0;
+
 	/**
 	 * Checks whether this device is being shared with other users.
 	 *
diff --git a/Engine/Source/Developer/TargetDeviceServices/Public/TargetDeviceServiceMessages.h b/Engine/Source/Developer/TargetDeviceServices/Public/TargetDeviceServiceMessages.h
--- a/Engine/Source/Developer/TargetDeviceServices/Public/TargetDeviceServiceMessages.h
+++ b/Engine/Source/Developer/TargetDeviceServices/Public/TargetDeviceServiceMessages.h
@@ -432,6 +432,10 @@ struct FTargetDeviceServicePong
 	UPROPERTY(EditAnywhere, Category="Message")
 	bool Connected;
 
+	/** Holds a flag indicating whether the device is authorized. */
+	UPROPERTY(EditAnywhere, Category="Message")
+	bool Authorized;
+
 	/** Holds the name of the host computer that the device is attached to. */
 	UPROPERTY(EditAnywhere, Category="Message")
 	FString HostName;
diff --git a/Engine/Source/Editor/UnrealEd/Private/Kismet2/DebuggerCommands.cpp b/Engine/Source/Editor/UnrealEd/Private/Kismet2/DebuggerCommands.cpp
--- a/Engine/Source/Editor/UnrealEd/Private/Kismet2/DebuggerCommands.cpp
+++ b/Engine/Source/Editor/UnrealEd/Private/Kismet2/DebuggerCommands.cpp
@@ -1053,6 +1053,10 @@ TSharedRef< SWidget > FPlayWorldCommands::GenerateLaunchMenuContent( TSharedRef<
 						TooltipArguments.Add(TEXT("DeviceID"), FText::FromString(DeviceProxy->GetName()));
 						TooltipArguments.Add(TEXT("DisplayName"), VanillaPlatform.PlatformInfo->DisplayName);
 						FText Tooltip = FText::Format(LOCTEXT("LaunchDeviceToolTipText", "Launch the game on this {DisplayName} device ({DeviceID})"), TooltipArguments);
+						if (!DeviceProxy->IsAuthorized())
+						{
+							Tooltip = FText::Format(LOCTEXT("LaunchDeviceToolTipText", "{DisplayName} device ({DeviceID}) is unauthorized or locked"), TooltipArguments);
+						}
 
 						FProjectStatus ProjectStatus;
 						if (IProjectManager::Get().QueryStatusForCurrentProject(ProjectStatus) && !ProjectStatus.IsTargetPlatformSupported(VanillaPlatform.PlatformInfo->VanillaPlatformName)) 
@@ -2416,7 +2420,7 @@ bool FInternalPlayWorldCommandCallbacks::CanLaunchOnDevice(const FString& Device
 		if (DeviceProxyManager.IsValid())
 		{
 			TSharedPtr<ITargetDeviceProxy> DeviceProxy = DeviceProxyManager->FindProxy(DeviceName);
-			if (DeviceProxy.IsValid() && DeviceProxy->IsConnected())
+			if (DeviceProxy.IsValid() && DeviceProxy->IsConnected() && DeviceProxy->IsAuthorized())
 			{
 				return true;
 			}
@@ -2429,7 +2433,7 @@ bool FInternalPlayWorldCommandCallbacks::CanLaunchOnDevice(const FString& Device
 			for (auto DevicesIt = Devices.CreateIterator(); DevicesIt; ++DevicesIt)
 			{
 				TSharedPtr<ITargetDeviceProxy> DeviceAggregateProxy = *DevicesIt;
-				if (DeviceAggregateProxy.IsValid() && DeviceAggregateProxy->IsConnected())
+				if (DeviceAggregateProxy.IsValid() && DeviceAggregateProxy->IsConnected() && DeviceAggregateProxy->IsAuthorized())
 				{
 					return true;
 				}
diff --git a/Engine/Source/Programs/AutomationTool/Lumin/LuminPlatform.Automation.cs b/Engine/Source/Programs/AutomationTool/Lumin/LuminPlatform.Automation.cs
--- a/Engine/Source/Programs/AutomationTool/Lumin/LuminPlatform.Automation.cs
+++ b/Engine/Source/Programs/AutomationTool/Lumin/LuminPlatform.Automation.cs
@@ -341,13 +341,20 @@ public class LuminPlatform : Platform
 				"@if \"%ERRORLEVEL%\" NEQ \"0\" goto Error",
 				"@echo.",
 				"@echo Installation successful",
+				"%MLDB% %DEVICE% ps > nul",
+				"@if \"%ERRORLEVEL%\" NEQ \"0\" goto OobeError",
 				"goto:eof",
+				":OobeError",
+				"@echo Device is not ready for use. Run \"%MLSDK% ps\" from a command prompt for details.",
+				"goto Pause",
 				":Error",
 				"@echo.",
 				"@echo There was an error installing the game. Look above for more info.",
 				"@echo.",
 				"@echo Things to try:",
 				"@echo Check that the device (and only the device) is listed with \"%MLDB% devices\" from a command prompt.",
+				"@echo Check if the device is ready for use with \"%MLSDK% ps\" from a command prompt.",
+				":Pause",
 				"@pause"
 			};
 		}
@@ -370,11 +377,16 @@ public class LuminPlatform : Platform
 					"\techo",
 					"\techo \"Things to try:\"",
 					"\techo \"Check that the device (and only the device) is listed with \"$MLDB devices\" from a command prompt.\"",
+					"\techo \"Check if the device is ready for use with \"%MLSDK% ps\" from a command prompt.\"",
 					"\techo",
 					"\texit 1",
 				"fi",
 				"echo",
 				"echo \"Installation successful\"",
+				"$MLDB $DEVICE ps > /dev/null",
+				"if [ $? -ne 0 ]; then",
+					"\techo \"Device is not ready for use. Run \"%MLSDK% ps\" from a command prompt for details.\"",				
+				"fi",
 				"exit 0",
 			};
 		}
@@ -541,6 +553,9 @@ public class LuminPlatform : Platform
 
 		Deploy.InitUPL(Params.ShortProjectName, SC.ProjectRoot, SC.StageTargets[0].Receipt.Configuration);
 
+		string MpkName = GetFinalMpkName(Params, SC);
+		string PackageName = GetPackageName(Params);
+
 		if (!Params.Prebuilt)
 		{
 			string additionalFiles = Deploy.StageFiles();
@@ -606,18 +621,15 @@ public class LuminPlatform : Platform
 
 			// now put the exe into the data list (coming from original location)
 			string ExePath = GetExePath(Params, SC);
-			Builder.AppendLine(String.Format("\"{0}\" : \"bin/{1}\"", ExePath, Params.ShortProjectName));
+			Builder.AppendLine(String.Format("\"{0}/Binaries/{1}\" : \"bin/{2}\"", Path.GetDirectoryName(MabuFile), Path.GetFileName(ExePath), Params.ShortProjectName));
 
 			Directory.CreateDirectory(Path.GetDirectoryName(MabuFile));
 			WriteAllText(MabuFile, Builder.ToString());
 
 			string ElfName = GetElfNameWithoutArchitecture(Params, SC.StageExecutables[0]);
-			Deploy.PrepForUATPackageOrDeploy(Params.RawProjectPath, Params.ShortProjectName, SC.ProjectRoot, ElfName, SC.LocalRoot + "/Engine", Params.Distribution, "", false);
+			Deploy.PrepForUATPackageOrDeploy(Params.RawProjectPath, Params.ShortProjectName, SC.ProjectRoot, ExePath, SC.LocalRoot + "/Engine", Params.Distribution, "", false, MpkName);
 		}
 
-		string MpkName = GetFinalMpkName(Params, SC);
-		string PackageName = GetPackageName(Params);
-
 		// Write install batch file(s).
 		string BatchName = GetFinalBatchName(Params, SC, false);
 		// make a batch file that can be used to install the .mpk and .obb files
@@ -825,7 +837,11 @@ public class LuminPlatform : Platform
 					// install the package only if was (re-)created during the package step
 					if (!IsPackageUpToDate(Params, SC))
 					{
-						RunDeviceCommand(Params, DeviceName, string.Format("install -u \"{0}\"", GetFinalMpkName(Params, SC)), null);
+						IProcessResult InstallResult = RunDeviceCommand(Params, DeviceName, string.Format("install -u \"{0}\"", GetFinalMpkName(Params, SC)), null);
+						if (InstallResult.ExitCode != 0)
+						{
+							throw new AutomationException((ExitCode)InstallResult.ExitCode, "Failed to install {0}", GetFinalMpkName(Params, SC));
+						}
 					}
 
 					// cache the timestamp of the exe
@@ -901,7 +917,11 @@ public class LuminPlatform : Platform
 		RunDeviceCommand(Params, DeviceName, UninstallCmd, null);
 
 		string InstallCmd = string.Format("install \"{0}\"", MpkName);
-		RunDeviceCommand(Params, DeviceName, InstallCmd, null);
+		IProcessResult InstallResult = RunDeviceCommand(Params, DeviceName, InstallCmd, null);
+		if (InstallResult.ExitCode != 0)
+		{
+			throw new AutomationException((ExitCode)InstallResult.ExitCode, "Failed to install {0}", GetFinalMpkName(Params, SC));
+		}
 	}
 
 	// adb shell quoted arguments need to be wrapped twice - once for Windows, and once for /bin/sh, otherwise sh can get confused
diff --git a/Engine/Source/Programs/GitDependencies/Program.cs b/Engine/Source/Programs/GitDependencies/Program.cs
--- a/Engine/Source/Programs/GitDependencies/Program.cs
+++ b/Engine/Source/Programs/GitDependencies/Program.cs
@@ -182,7 +182,7 @@ namespace GitDependencies
 				{
 					ExcludeFolders.Add("Mac");
 				}
-				if(Environment.GetEnvironmentVariable("NDKROOT") == null)
+				if(Environment.GetEnvironmentVariable("NDKROOT") == null && Environment.GetEnvironmentVariable("MLSDK") == null)
 				{
 					ExcludeFolders.Add("Android");
 				}
diff --git a/Engine/Source/Programs/UnrealBuildTool/Platform/Android/AndroidToolChain.cs b/Engine/Source/Programs/UnrealBuildTool/Platform/Android/AndroidToolChain.cs
--- a/Engine/Source/Programs/UnrealBuildTool/Platform/Android/AndroidToolChain.cs
+++ b/Engine/Source/Programs/UnrealBuildTool/Platform/Android/AndroidToolChain.cs
@@ -1826,7 +1826,7 @@ namespace UnrealBuildTool
 
 			ProcessStartInfo StartInfo = new ProcessStartInfo();
 			StartInfo.FileName = GetStripPath(SourceFile);
-			StartInfo.Arguments = "--strip-debug " + TargetFile.FullName;
+			StartInfo.Arguments = " --strip-debug \"" + TargetFile.FullName + "\"";
 			StartInfo.UseShellExecute = false;
 			StartInfo.CreateNoWindow = true;
 			Utils.RunLocalProcessAndLogOutput(StartInfo);
diff --git a/Engine/Source/Programs/UnrealBuildTool/Platform/Lumin/LuminExports.cs b/Engine/Source/Programs/UnrealBuildTool/Platform/Lumin/LuminExports.cs
--- a/Engine/Source/Programs/UnrealBuildTool/Platform/Lumin/LuminExports.cs
+++ b/Engine/Source/Programs/UnrealBuildTool/Platform/Lumin/LuminExports.cs
@@ -47,8 +47,9 @@ namespace UnrealBuildTool
 		/// <param name="bForDistribution"></param>
 		/// <param name="CookFlavor"></param>
 		/// <param name="bIsDataDeploy"></param>
+		/// <param name="MpkName"></param>
 		/// <returns></returns>
-		bool PrepForUATPackageOrDeploy(FileReference ProjectFile, string ProjectName, DirectoryReference ProjectDirectory, string ExecutablePath, string EngineDirectory, bool bForDistribution, string CookFlavor, bool bIsDataDeploy);
+		bool PrepForUATPackageOrDeploy(FileReference ProjectFile, string ProjectName, DirectoryReference ProjectDirectory, string ExecutablePath, string EngineDirectory, bool bForDistribution, string CookFlavor, bool bIsDataDeploy, string MpkName);
 
 		/// <summary>
 		/// 
diff --git a/Engine/Source/Programs/UnrealBuildTool/Platform/Lumin/LuminToolChain.cs b/Engine/Source/Programs/UnrealBuildTool/Platform/Lumin/LuminToolChain.cs
--- a/Engine/Source/Programs/UnrealBuildTool/Platform/Lumin/LuminToolChain.cs
+++ b/Engine/Source/Programs/UnrealBuildTool/Platform/Lumin/LuminToolChain.cs
@@ -232,6 +232,10 @@ namespace UnrealBuildTool
 			Result += " -Wl,-rpath=$ORIGIN";
 			Result += " -fdiagnostics-format=msvc";
 
+			if (!LinkEnvironment.bCreateDebugInfo)
+			{
+				Result += " -Wl,--strip-debug";
+			}
 
 			if (!LinkEnvironment.bIsBuildingDLL)
 			{
diff --git a/Engine/Source/Programs/UnrealBuildTool/Platform/Lumin/UEBuildLumin.cs b/Engine/Source/Programs/UnrealBuildTool/Platform/Lumin/UEBuildLumin.cs
--- a/Engine/Source/Programs/UnrealBuildTool/Platform/Lumin/UEBuildLumin.cs
+++ b/Engine/Source/Programs/UnrealBuildTool/Platform/Lumin/UEBuildLumin.cs
@@ -270,7 +270,7 @@ namespace UnrealBuildTool
 		/// <summary>
 		/// This is the SDK version we support
 		/// </summary>
-		static string ExpectedSDKVersion = "0.12";   // now unified for all the architectures
+		static string ExpectedSDKVersion = "0.15";   // now unified for all the architectures
 
 		public override string GetSDKTargetPlatformName()
 		{
diff --git a/Engine/Source/Programs/UnrealBuildTool/Platform/Lumin/UEDeployLumin.cs b/Engine/Source/Programs/UnrealBuildTool/Platform/Lumin/UEDeployLumin.cs
--- a/Engine/Source/Programs/UnrealBuildTool/Platform/Lumin/UEDeployLumin.cs
+++ b/Engine/Source/Programs/UnrealBuildTool/Platform/Lumin/UEDeployLumin.cs
@@ -5,6 +5,7 @@ using System.Collections.Generic;
 using System.Diagnostics;
 using System.IO;
 using System.Text;
+using System.Text.RegularExpressions;
 using Tools.DotNETCommon;
 using System.Linq;
 using System.Xml;
@@ -76,6 +77,33 @@ namespace UnrealBuildTool
 			return ApplicationDisplayName;
 		}
 
+        private string GetMinimumOSVersionRequired()
+        {
+            string MinOSVersionRequired = GetRuntimeSetting("MinimumOSVersion");
+            return String.IsNullOrWhiteSpace(MinOSVersionRequired) ? "nova_1.0" : MinOSVersionRequired;
+        }
+
+        private void GetAppPrivileges(ConfigHierarchy EngineIni, StringBuilder Text)
+        {
+			List<string> AppPrivileges;
+			EngineIni.GetArray("/Script/LuminRuntimeSettings.LuminRuntimeSettings", "AppPrivileges", out AppPrivileges);
+			if (AppPrivileges != null)
+			{
+				foreach (string Privilege in AppPrivileges)
+				{
+					string TrimmedPrivilege = Privilege.Trim(' ');
+					if (TrimmedPrivilege != "")
+					{
+						string PrivilegeString = string.Format("\t\t\t<uses-privilege ml:name=\"{0}\"/>", TrimmedPrivilege);
+						if (!Text.ToString().Contains(PrivilegeString))
+						{
+							Text.AppendLine(PrivilegeString);
+						}
+					}
+				}
+			}
+        }
+
 		private string CleanFilePath(string FilePath)
 		{
 			// Removes the extra characters from a FFilePath parameter.
@@ -96,6 +124,45 @@ namespace UnrealBuildTool
 			return "Icon/Portal";
 		}
 
+        public string GetMLSDKVersion(ConfigHierarchy EngineIni)
+        {
+            string MLSDKPath;
+            string Major = "0";
+            string Minor = "0";
+            if (EngineIni.GetString("/Script/LuminPlatformEditor.MagicLeapSDKSettings", "MLSDKPath", out MLSDKPath) && !string.IsNullOrWhiteSpace(MLSDKPath))
+            {
+                Dictionary<string, string> PathEntry;
+                ConfigHierarchy.TryParse(MLSDKPath, out PathEntry);
+                MLSDKPath = PathEntry.Values.First();
+            }
+            if (string.IsNullOrWhiteSpace(MLSDKPath))
+            {
+                MLSDKPath = Environment.GetEnvironmentVariable("MLSDK");
+            }
+            if (!string.IsNullOrEmpty(MLSDKPath))
+            {
+                if (Directory.Exists(MLSDKPath))
+                {
+                    String VersionFile = string.Format("{0}/include/ml_version.h", MLSDKPath).Replace('/', Path.DirectorySeparatorChar);
+                    if (File.Exists(VersionFile))
+                    {
+                        string FileText = File.ReadAllText(VersionFile);
+                        string Pattern = @"(MLSDK_VERSION_MAJOR) (?'MAJOR'\d+).*(MLSDK_VERSION_MINOR) (?'MINOR'\d+).*(MLSDK_VERSION_REVISION) (?'REV'\d+)";
+                        Regex VersionRegex = new Regex(Pattern, RegexOptions.Singleline);
+                        MatchCollection Matches = VersionRegex.Matches(FileText);
+                        if (Matches.Count > 0 &&
+                            !string.IsNullOrEmpty(Matches[0].Groups["MAJOR"].Value) &&
+                            !string.IsNullOrEmpty(Matches[0].Groups["MINOR"].Value))
+                        {
+                            Major = Matches[0].Groups["MAJOR"].Value;
+                            Minor = Matches[0].Groups["MINOR"].Value;
+                        }
+                    }
+                }
+            }
+            return string.Format("{0}.{1}", Major, Minor);
+        }
+
 		public string GenerateManifest(string ProjectName, bool bForDistribution, string Architecture)
 		{
 			ConfigHierarchy GameIni = GetConfigCacheIni(ConfigHierarchyType.Game);
@@ -110,39 +177,24 @@ namespace UnrealBuildTool
 			int VersionCode;
 			EngineIni.GetInt32("/Script/LuminRuntimeSettings.LuminRuntimeSettings", "VersionCode", out VersionCode);
 
-			bool bInternetRequired;
-			EngineIni.GetBool("/Script/LuminRuntimeSettings.LuminRuntimeSettings", "bInternetRequired", out bInternetRequired);
+			string SDKVersion = GetMLSDKVersion(EngineIni);
 
 			StringBuilder Text = new StringBuilder();
 
 			string PackageName = GetPackageName(ProjectName);
 			string ApplicationDisplayName = GetApplicationDisplayName(ProjectName);
+			string MinimumOSVersion = GetMinimumOSVersionRequired();
 			string TargetExecutableName = "bin/" + ProjectName;
 
 			Text.AppendLine(string.Format("<manifest xmlns:ml=\"magicleap\" ml:package=\"{0}\" ml:version_name=\"{1}\" ml:version_code=\"{2}\">", PackageName, ProjectVersion, VersionCode));
-			// @mltodo: query sdk_version
-			Text.AppendLine(string.Format("\t<application ml:visible_name=\"{0}\" ml:is_debuggable=\"{1}\" ml:sdk_version=\"1.0\" ml:minimum_os=\"mlos_1.0\" ml:internet_required=\"{2}\">", ApplicationDisplayName, bForDistribution ? "false" : "true", bInternetRequired ? "true" : "false"));
+			Text.AppendLine(string.Format("\t<application ml:visible_name=\"{0}\" ml:sdk_version=\"{1}\" ml:minimum_os=\"{2}\">",
+                ApplicationDisplayName,
+                SDKVersion,
+                MinimumOSVersion));
+            GetAppPrivileges(EngineIni, Text);
 			Text.AppendLine(string.Format("\t\t<component ml:name=\".fullscreen\" ml:visible_name=\"{0}\" ml:binary_name=\"{1}\" ml:type=\"{2}\">", ApplicationDisplayName, TargetExecutableName, GetApplicationType()));
 
-			List<string> AppPrivileges;
-			EngineIni.GetArray("/Script/LuminRuntimeSettings.LuminRuntimeSettings", "AppPrivileges", out AppPrivileges);
-			if (AppPrivileges != null)
-			{
-				foreach (string Privilege in AppPrivileges)
-				{
-					string TrimmedPrivilege = Privilege.Trim(' ');
-					if (TrimmedPrivilege != "")
-					{
-						string PrivilegeString = string.Format("\t\t\t<uses-privilege ml:name=\"{0}\"/>", TrimmedPrivilege);
-						if (!Text.ToString().Contains(PrivilegeString))
-						{
-							Text.AppendLine(PrivilegeString);
-						}
-					}
-				}
-			}
-
-			string IconTag = string.Format("<icon ml:name=\"fullscreen\" ml:model_folder=\"{0}\" ml:portal_folder=\"{1}\"/>", GetIconModelStagingPath(), GetIconPortalStagingPath());
+			string IconTag = string.Format("<icon ml:model_folder=\"{0}\" ml:portal_folder=\"{1}\"/>", GetIconModelStagingPath(), GetIconPortalStagingPath());
 			Text.AppendLine(string.Format("\t\t\t{0}", IconTag));
 
 			List<string> ExtraComponentNodes;
@@ -167,7 +219,6 @@ namespace UnrealBuildTool
 				}
 			}
 
-			Text.AppendLine(string.Format("\t\t{0}", IconTag));
 			Text.AppendLine("\t</application>");
 			Text.AppendLine("</manifest>");
 
@@ -245,7 +296,7 @@ namespace UnrealBuildTool
 			return UPL.ProcessPluginNode(Architecture, "stageFiles", "");
 		}
 
-		public void MakeMabuPackage(string ProjectName, DirectoryReference ProjectDirectory, string ExePath, bool bForDistribution, string EngineDir)
+		private void MakeMabuPackage(string ProjectName, DirectoryReference ProjectDirectory, string ExePath, bool bForDistribution, string EngineDir, string MpkName)
 		{
 			string UE4BuildPath = Path.Combine(ProjectDirectory.FullName, "Intermediate/Lumin/Mabu");
 			string MabuOutputPath = Path.Combine(UE4BuildPath, "Packaged");
@@ -256,6 +307,20 @@ namespace UnrealBuildTool
 
 			LuminToolChain ToolChain = new LuminToolChain(ProjectFile);
 
+			// If asked for, and if we are doing a distribution package, we strip debug symbols.
+			string ExecSrcFile = Path.Combine(UE4BuildPath, "Binaries", Path.GetFileName(ExePath));
+			if (bForDistribution)
+			{
+				Directory.CreateDirectory(Path.GetDirectoryName(ExecSrcFile));
+				ToolChain.StripSymbols(new Tools.DotNETCommon.FileReference(ExePath), new Tools.DotNETCommon.FileReference(ExecSrcFile));
+			}
+			else
+			{
+				// The generated mabu needs the src exe file. So we copy the original as-is so mabu can find it.
+				Directory.CreateDirectory(Path.GetDirectoryName(ExecSrcFile));
+				File.Copy(ExePath, ExecSrcFile, true);
+			}
+
 			// Generate manifest (after UPL is setup
 			string Architecture = "arm64-v8a";
 			var Manifest = GenerateManifest(ProjectName, bForDistribution, Architecture);
@@ -296,28 +361,22 @@ namespace UnrealBuildTool
 			// @todo Lumin: This is losing the -Debug-Lumin stuff :|
 			string SourceMpkPath = Path.Combine(MabuOutputPath, GetPackageName(ProjectName) + ".mpk");
 
-			string MpkName = Path.GetFileNameWithoutExtension(ExePath) + ".mpk";
-			if (MpkName.StartsWith("UE4Game"))
-			{
-				MpkName = MpkName.Replace("UE4Game", ProjectName);
-			}
-			string DestMpkPath = Path.Combine(ProjectDirectory.FullName, "Binaries/Lumin", MpkName);
-			if (!Directory.Exists(Path.GetDirectoryName(DestMpkPath)))
+			if (!Directory.Exists(Path.GetDirectoryName(MpkName)))
 			{
-				Directory.CreateDirectory(Path.GetDirectoryName(DestMpkPath));
+				Directory.CreateDirectory(Path.GetDirectoryName(MpkName));
 			}
-			if (File.Exists(DestMpkPath))
+			if (File.Exists(MpkName))
 			{
-				File.Delete(DestMpkPath);
+				File.Delete(MpkName);
 			}
-			File.Copy(SourceMpkPath, DestMpkPath);
+			File.Copy(SourceMpkPath, MpkName);
 		}
 
-		public bool PrepForUATPackageOrDeploy(FileReference ProjectFile, string InProjectName, DirectoryReference InProjectDirectory, string InExecutablePath, string InEngineDir, bool bForDistribution, string CookFlavor, bool bIsDataDeploy)
+		public bool PrepForUATPackageOrDeploy(FileReference ProjectFile, string InProjectName, DirectoryReference InProjectDirectory, string InExecutablePath, string InEngineDir, bool bForDistribution, string CookFlavor, bool bIsDataDeploy, string MpkName)
 		{
 			if (!bIsDataDeploy)
 			{
-				MakeMabuPackage(InProjectName, InProjectDirectory, InExecutablePath, bForDistribution, InEngineDir);
+				MakeMabuPackage(InProjectName, InProjectDirectory, InExecutablePath, bForDistribution, InEngineDir, MpkName);
 			}
 			return true;
 		}
diff --git a/Engine/Source/Runtime/AudioMixer/Private/AudioMixerSource.cpp b/Engine/Source/Runtime/AudioMixer/Private/AudioMixerSource.cpp
--- a/Engine/Source/Runtime/AudioMixer/Private/AudioMixerSource.cpp
+++ b/Engine/Source/Runtime/AudioMixer/Private/AudioMixerSource.cpp
@@ -353,7 +353,7 @@ namespace Audio
 			bResourcesNeedFreeing = (BufferType == EBufferType::PCMRealTime || BufferType == EBufferType::Streaming);
 
 			// Not all wave data types have a non-zero duration
-			if (InWaveInstance->WaveData->Duration > 0)
+			if (!InWaveInstance->WaveData->IsLooping())
 			{
 				if (!InWaveInstance->WaveData->bIsBus)
 				{
diff --git a/Engine/Source/Runtime/AudioMixer/Private/Components/SynthComponent.cpp b/Engine/Source/Runtime/AudioMixer/Private/Components/SynthComponent.cpp
--- a/Engine/Source/Runtime/AudioMixer/Private/Components/SynthComponent.cpp
+++ b/Engine/Source/Runtime/AudioMixer/Private/Components/SynthComponent.cpp
@@ -157,52 +157,55 @@ void USynthComponent::Deactivate()
 
 void USynthComponent::Initialize(int32 SampleRateOverride)
 {
-	if (!bIsInitialized)
+	// Try to get a proper sample rate
+	int32 SampleRate = SampleRateOverride;
+	if (SampleRate == INDEX_NONE)
 	{
-		// Try to get a proper sample rate
-		int32 SampleRate = SampleRateOverride;
-		if (SampleRate == INDEX_NONE)
+		// Check audio device if we've not explicitly been told what sample rate to use
+		FAudioDevice* AudioDevice = GetAudioDevice();
+		if (AudioDevice)
 		{
-			// Check audio device if we've not explicitly been told what sample rate to use
-			FAudioDevice* AudioDevice = GetAudioDevice();
-			if (AudioDevice)
-			{
-				SampleRate = AudioDevice->SampleRate;
-			}
+			SampleRate = AudioDevice->SampleRate;
 		}
+	}
 
-		// Only allow initialization if we've gota  proper sample rate
-		if (SampleRate != INDEX_NONE)
-		{
-			bIsInitialized = true;
-
+	// Only allow initialization if we've gota  proper sample rate
+	if (SampleRate != INDEX_NONE)
+	{
 #if SYNTH_GENERATOR_TEST_TONE
-			NumChannels = 2;
-			TestSineLeft.Init(SampleRate, 440.0f, 0.5f);
-			TestSineRight.Init(SampleRate, 220.0f, 0.5f);
+		NumChannels = 2;
+		TestSineLeft.Init(SampleRate, 440.0f, 0.5f);
+		TestSineRight.Init(SampleRate, 220.0f, 0.5f);
 #else	
-			// Initialize the synth component
-			this->Init(SampleRate);
+		// Initialize the synth component
+		this->Init(SampleRate);
 
-			if (NumChannels < 0 || NumChannels > 2)
-			{
-				UE_LOG(LogAudioMixer, Error, TEXT("Synthesis component '%s' has set an invalid channel count '%d' (only mono and stereo currently supported)."), *GetName(), NumChannels);
-			}
+		if (NumChannels < 0 || NumChannels > 2)
+		{
+			UE_LOG(LogAudioMixer, Error, TEXT("Synthesis component '%s' has set an invalid channel count '%d' (only mono and stereo currently supported)."), *GetName(), NumChannels);
+		}
 
-			NumChannels = FMath::Clamp(NumChannels, 1, 2);
+		NumChannels = FMath::Clamp(NumChannels, 1, 2);
 #endif
 
+		if (nullptr == Synth)
+		{
 			Synth = NewObject<USynthSound>(this, TEXT("Synth"));
+		}
 
-			// Copy sound base data to the sound
-			Synth->SourceEffectChain = SourceEffectChain;
-			Synth->SoundSubmixObject = SoundSubmix;
-			Synth->SoundSubmixSends = SoundSubmixSends;
-			Synth->BusSends = BusSends;
-			Synth->PreEffectBusSends = PreEffectBusSends;
-			Synth->bOutputToBusOnly = bOutputToBusOnly;
+		// Copy sound base data to the sound
+		Synth->SourceEffectChain = SourceEffectChain;
+		Synth->SoundSubmixObject = SoundSubmix;
+		Synth->SoundSubmixSends = SoundSubmixSends;
+		Synth->BusSends = BusSends;
+		Synth->PreEffectBusSends = PreEffectBusSends;
+		Synth->bOutputToBusOnly = bOutputToBusOnly;
+
+		Synth->Init(this, NumChannels, SampleRate, PreferredBufferLength);
 
-			Synth->Init(this, NumChannels, SampleRate, PreferredBufferLength);
+		if (FAudioDevice* AudioDevice = AudioComponent->GetAudioDevice())
+		{
+			Synth->StartOnAudioDevice(AudioDevice);
 		}
 	}
 }
@@ -239,13 +242,6 @@ void USynthComponent::CreateAudioComponent()
 			// Set defaults to be the same as audio component defaults
 			AudioComponent->EnvelopeFollowerAttackTime = EnvelopeFollowerAttackTime;
 			AudioComponent->EnvelopeFollowerReleaseTime = EnvelopeFollowerReleaseTime;
-
-			Initialize();
-
-			if (FAudioDevice* AudioDevice = AudioComponent->GetAudioDevice())
-			{
-				Synth->StartOnAudioDevice(AudioDevice);
-			}
 		}
 	}
 }
@@ -253,8 +249,6 @@ void USynthComponent::CreateAudioComponent()
 
 void USynthComponent::OnRegister()
 {
-	CreateAudioComponent();
-
 	Super::OnRegister();
 }
 
diff --git a/Engine/Source/Runtime/Core/Private/HAL/MallocBinned2.cpp b/Engine/Source/Runtime/Core/Private/HAL/MallocBinned2.cpp
--- a/Engine/Source/Runtime/Core/Private/HAL/MallocBinned2.cpp
+++ b/Engine/Source/Runtime/Core/Private/HAL/MallocBinned2.cpp
@@ -12,7 +12,7 @@
 #include "Stats/Stats.h"
 #include "HAL/IConsoleManager.h"
 #include "HAL/MemoryMisc.h"
-
+#include "HAL/PlatformMisc.h"
 
 #if BINNED2_ALLOW_RUNTIME_TWEAKING
 
@@ -236,7 +236,15 @@ private:
 		{
 			FirstFreeBlock = (FFreeBlock*)FirstFreeBlock->NextFreeBlock;
 		}
+#if PLATFORM_LUMIN
+		// Temporarily disable fatal check for Lumin since it crashes the ImageTrackerComponent.
+		if (!(!FirstFreeBlock || FirstFreeBlock->GetNumFreeRegularBlocks() != 0))
+		{
+			FPlatformMisc::LowLevelOutputDebugString(TEXT("assert failed - (!FirstFreeBlock || FirstFreeBlock->GetNumFreeRegularBlocks() != 0)"));
+		}
+#else
 		check(!FirstFreeBlock || FirstFreeBlock->GetNumFreeRegularBlocks() != 0);
+#endif // PLATFORM_LUMIN
 	}
 };
 
diff --git a/Engine/Source/Runtime/Core/Private/Lumin/LuminLifecycle.cpp b/Engine/Source/Runtime/Core/Private/Lumin/LuminLifecycle.cpp
--- a/Engine/Source/Runtime/Core/Private/Lumin/LuminLifecycle.cpp
+++ b/Engine/Source/Runtime/Core/Private/Lumin/LuminLifecycle.cpp
@@ -10,7 +10,7 @@ DEFINE_LOG_CATEGORY(LogLifecycle);
 
 bool FLuminLifecycle::bIsEngineLoopInitComplete = false;
 bool FLuminLifecycle::bIsAppPaused = false;
-MLLifecycleErrorCode FLuminLifecycle::LifecycleState = MLLifecycleErrorCode_Failure;
+MLResult FLuminLifecycle::LifecycleState = MLResult_UnspecifiedFailure;
 MLLifecycleCallbacks FLuminLifecycle::LifecycleCallbacks = { nullptr, nullptr, nullptr, nullptr, nullptr };
 TArray<FString> FLuminLifecycle::PendingStartupArgs;
 
@@ -34,7 +34,7 @@ void FLuminLifecycle::Initialize()
 	// TODO: confirm this comment for ml_lifecycle.
 	// There's a known issue where ck_lifecycle_init will fail to initialize if the debugger is attached.
 	// Ideally, this should assert since the app won't be able to react to events correctly.
-	if (LifecycleState == MLLifecycleErrorCode_Failure)
+	if (LifecycleState != MLResult_Ok)
 	{
 		FPlatformMisc::LowLevelOutputDebugString(TEXT("Lifecycle system failed to initialize! App may not suspend, resume, or teriminate correctly."));
 	}
@@ -46,7 +46,7 @@ void FLuminLifecycle::Initialize()
 
 bool FLuminLifecycle::IsLifecycleInitialized()
 {
-	return (LifecycleState == MLLifecycleErrorCode_Success);
+	return (LifecycleState == MLResult_Ok);
 }
 
 void FLuminLifecycle::Stop_Handler(void* ApplicationContext)
@@ -132,20 +132,24 @@ void FLuminLifecycle::UnloadResources_Handler(void* ApplicationContext)
 
 void FLuminLifecycle::OnNewInitArgs_Handler(void* ApplicationContext)
 {
-	MLLifecycleInitArgList* InitArgList = MLLifecycleGetInitArgList();
-	if (InitArgList != nullptr)
+	MLLifecycleInitArgList* InitArgList = nullptr;
+	MLResult Result = MLLifecycleGetInitArgList(&InitArgList);
+	if (Result == MLResult_Ok && InitArgList != nullptr)
 	{
-		int64 InitArgCount = MLLifecycleGetInitArgListLength(InitArgList);
-		if (InitArgCount > 0)
+		int64_t InitArgCount = 0;
+		Result = MLLifecycleGetInitArgListLength(InitArgList, &InitArgCount);
+		if (Result == MLResult_Ok && InitArgCount > 0)
 		{
 			TArray<FString> InitArgsArray;
 			for (int64 i = 0; i < InitArgCount; ++i)
 			{
-				const MLLifecycleInitArg* InitArg = MLLifecycleGetInitArgByIndex(InitArgList, i);
-				if (InitArg != nullptr)
+				const MLLifecycleInitArg* InitArg = nullptr;
+				Result = MLLifecycleGetInitArgByIndex(InitArgList, i, &InitArg);
+				if (Result == MLResult_Ok && InitArg != nullptr)
 				{
-					const char* Arg = MLLifecycleGetInitArgUri(InitArg);
-					if (Arg != nullptr)
+					const char* Arg = nullptr;
+					Result = MLLifecycleGetInitArgUri(InitArg, &Arg);
+					if (Result == MLResult_Ok && Arg != nullptr)
 					{
 						// Start with a space because the command line already in place may not have any trailing spaces.
 						FString ArgStr = TEXT(" ");
diff --git a/Engine/Source/Runtime/Core/Private/Lumin/LuminLifecycle.h b/Engine/Source/Runtime/Core/Private/Lumin/LuminLifecycle.h
--- a/Engine/Source/Runtime/Core/Private/Lumin/LuminLifecycle.h
+++ b/Engine/Source/Runtime/Core/Private/Lumin/LuminLifecycle.h
@@ -23,7 +23,7 @@ private:
 private:
 	static bool bIsEngineLoopInitComplete;
 	static bool bIsAppPaused;
-	static MLLifecycleErrorCode LifecycleState;
+	static MLResult LifecycleState;
 	static MLLifecycleCallbacks LifecycleCallbacks;
 	static TArray<FString> PendingStartupArgs;
 };
