Fix for vk media player crash

diff --git a/Engine/Plugins/Lumin/MagicLeapMedia/Source/MagicLeapMedia/Private/MagicLeapMediaPlayer.cpp b/Engine/Plugins/Lumin/MagicLeapMedia/Source/MagicLeapMedia/Private/MagicLeapMediaPlayer.cpp
--- a/Engine/Plugins/Lumin/MagicLeapMedia/Source/MagicLeapMedia/Private/MagicLeapMediaPlayer.cpp
+++ b/Engine/Plugins/Lumin/MagicLeapMedia/Source/MagicLeapMedia/Private/MagicLeapMediaPlayer.cpp
@@ -1,4 +1,4 @@
-﻿// Copyright 1998-2018 Epic Games, Inc. All Rights Reserved.
+// Copyright 1998-2018 Epic Games, Inc. All Rights Reserved.
 
 #include "MagicLeapMediaPlayer.h"
 #include "IMagicLeapMediaModule.h"
@@ -244,9 +244,9 @@ FMagicLeapMediaPlayer::~FMagicLeapMediaPlayer()
 				MLResult Result = MLMediaPlayerDestroy(Params.MediaPlayerHandle);
 				UE_CLOG(Result != MLResult_Ok, LogMagicLeapMedia, Error, TEXT("MLMediaPlayerDestroy failed with error %d."), Result);
 					});
-
-				FlushRenderingCommands();
 			}
+
+			FlushRenderingCommands();
 		}
 		else
 		{
