#DEV-VR - Merging CL 4199443 from Dev-Magicleap to clear cache on looping media

diff --git a/Engine/Plugins/Lumin/MagicLeapMedia/Source/MagicLeapMedia/Private/MagicLeapMediaPlayer.cpp b/Engine/Plugins/Lumin/MagicLeapMedia/Source/MagicLeapMedia/Private/MagicLeapMediaPlayer.cpp
--- a/Engine/Plugins/Lumin/MagicLeapMedia/Source/MagicLeapMedia/Private/MagicLeapMediaPlayer.cpp
+++ b/Engine/Plugins/Lumin/MagicLeapMedia/Source/MagicLeapMedia/Private/MagicLeapMediaPlayer.cpp
@@ -144,6 +144,7 @@ FMagicLeapMediaPlayer::FMagicLeapMediaPlayer(IMediaEventSink& InEventSink)
 , VideoSamplePool(new FMagicLeapMediaTextureSamplePool())
 , MediaWorker(nullptr)
 , bWasMediaPlayingBeforeAppPause(false)
+, bPlaybackCompleted(false)
 {
 	if (FLuminPlatformMisc::ShouldUseVulkan())
 	{
@@ -670,6 +671,15 @@ void FMagicLeapMediaPlayer::TickFetch(FTimespan DeltaTime, FTimespan Timecode)
 
 					check(MLHandleIsValid(NativeBuffer));
 
+					{
+						FScopeLock Lock(&Params.MediaPlayer->CriticalSection);
+						if (Params.MediaPlayer->bPlaybackCompleted)
+						{
+							TextureDataPtr->VideoTexturePool.Empty();
+							Params.MediaPlayer->bPlaybackCompleted = false;
+						}
+					}
+
 					if (!TextureDataPtr->VideoTexturePool.Contains((uint64)NativeBuffer))
 					{
 						FTextureRHIRef NewMediaTexture;
@@ -1000,6 +1010,10 @@ void FMagicLeapMediaPlayer::TickInput(FTimespan DeltaTime, FTimespan Timecode)
 
 	if (GetMediaPlayerState(MLMediaPlayerPollingStateFlag_HasPlaybackCompleted))
 	{
+		{
+			FScopeLock Lock(&CriticalSection);
+			bPlaybackCompleted = true;
+		}
 		if (!IsLooping())
 		{
 			CurrentState = EMediaState::Stopped;
diff --git a/Engine/Plugins/Lumin/MagicLeapMedia/Source/MagicLeapMedia/Private/MagicLeapMediaPlayer.h b/Engine/Plugins/Lumin/MagicLeapMedia/Source/MagicLeapMedia/Private/MagicLeapMediaPlayer.h
--- a/Engine/Plugins/Lumin/MagicLeapMedia/Source/MagicLeapMedia/Private/MagicLeapMediaPlayer.h
+++ b/Engine/Plugins/Lumin/MagicLeapMedia/Source/MagicLeapMedia/Private/MagicLeapMediaPlayer.h
@@ -151,6 +151,8 @@ protected:
 
 	bool bWasMediaPlayingBeforeAppPause;
 
+	bool bPlaybackCompleted;
+
 private:
 	virtual bool SetRateOne();
 	virtual bool GetMediaPlayerState(uint16 FlagToPoll) const;
