[NO NON-PII COMMIT MESSAGE]

diff --git a/Engine/Plugins/Lumin/MagicLeapMedia/Source/MagicLeapMediaCodec/Private/MagicLeapMediaCodecPlayer.cpp b/Engine/Plugins/Lumin/MagicLeapMedia/Source/MagicLeapMediaCodec/Private/MagicLeapMediaCodecPlayer.cpp
--- a/Engine/Plugins/Lumin/MagicLeapMedia/Source/MagicLeapMediaCodec/Private/MagicLeapMediaCodecPlayer.cpp
+++ b/Engine/Plugins/Lumin/MagicLeapMedia/Source/MagicLeapMediaCodec/Private/MagicLeapMediaCodecPlayer.cpp
@@ -750,7 +750,7 @@ void FMagicLeapMediaCodecPlayer::TickFetch(FTimespan DeltaTime, FTimespan Timeco
 			FWriteVideoSampleParams Params = { this, StaticCastSharedPtr<FMagicLeapVideoTextureDataVK>(TextureData), PlayerGuid, VideoCodecHandle, Timecode };
 
 			ENQUEUE_RENDER_COMMAND(MagicLeapMediaPlayerWriteVideoSample)(
-				[Params](FRHICommandListImmediate& RHICmdList)
+				[Params](FRHICommandListImmediate& RHICmdList) mutable
 				{
 					auto TextureDataPtr = Params.TextureData.Pin();
 
@@ -831,7 +831,7 @@ void FMagicLeapMediaCodecPlayer::TickFetch(FTimespan DeltaTime, FTimespan Timeco
 			FWriteVideoSampleParams Params = { this, StaticCastSharedPtr<FMagicLeapVideoTextureDataGL>(TextureData), PlayerGuid, VideoCodecHandle, Timecode };
 
 			ENQUEUE_RENDER_COMMAND(MagicLeapMediaPlayerWriteVideoSample)(
-				[Params](FRHICommandListImmediate& RHICmdList)
+				[Params](FRHICommandListImmediate& RHICmdList) mutable
 				{
 					auto TextureDataPtr = Params.TextureData.Pin();
 
