UE4.22 - Missed file
#rnx

diff --git a/Engine/Plugins/Lumin/MagicLeapMedia/Source/MagicLeapMedia/Private/MagicLeapMediaPlayer.cpp b/Engine/Plugins/Lumin/MagicLeapMedia/Source/MagicLeapMedia/Private/MagicLeapMediaPlayer.cpp
--- a/Engine/Plugins/Lumin/MagicLeapMedia/Source/MagicLeapMedia/Private/MagicLeapMediaPlayer.cpp
+++ b/Engine/Plugins/Lumin/MagicLeapMedia/Source/MagicLeapMedia/Private/MagicLeapMediaPlayer.cpp
@@ -188,10 +188,10 @@ FMagicLeapMediaPlayer::~FMagicLeapMediaPlayer()
 					MLHandle MediaPlayerHandle;
 				};
 
-				FReleaseVideoResourcesParams ReleaseVideoResourcesParams = { this, StaticCastSharedPtr<FMagicLeapVideoTextureDataVK>(TextureData), PlayerGuid, MediaPlayerHandle };
+				FReleaseVideoResourcesParams Params = { this, StaticCastSharedPtr<FMagicLeapVideoTextureDataVK>(TextureData), PlayerGuid, MediaPlayerHandle };
 
-				ENQUEUE_UNIQUE_RENDER_COMMAND_ONEPARAMETER(MagicLeapMediaPlayerDestroy,
-					FReleaseVideoResourcesParams, Params, ReleaseVideoResourcesParams,
+				ENQUEUE_RENDER_COMMAND(MagicLeapMediaPlayerDestroy)(
+					[Params](FRHICommandListImmediate& RHICmdList)
 					{
 						FExternalTextureRegistry::Get().UnregisterExternalTexture(Params.PlayerGuid);
 
@@ -220,31 +220,31 @@ FMagicLeapMediaPlayer::~FMagicLeapMediaPlayer()
 					MLHandle MediaPlayerHandle;
 				};
 
-				FReleaseVideoResourcesParams ReleaseVideoResourcesParams = { this, StaticCastSharedPtr<FMagicLeapVideoTextureDataGL>(TextureData), PlayerGuid, MediaPlayerHandle };
+				FReleaseVideoResourcesParams Params = { this, StaticCastSharedPtr<FMagicLeapVideoTextureDataGL>(TextureData), PlayerGuid, MediaPlayerHandle };
 
-				ENQUEUE_UNIQUE_RENDER_COMMAND_ONEPARAMETER(MagicLeapMediaPlayerDestroy,
-					FReleaseVideoResourcesParams, Params, ReleaseVideoResourcesParams,
+				ENQUEUE_RENDER_COMMAND(MagicLeapMediaPlayerDestroy)(
+					[Params](FRHICommandListImmediate& RHICmdList)
 					{
 						FExternalTextureRegistry::Get().UnregisterExternalTexture(Params.PlayerGuid);
-				// @todo: this causes a crash
-				//Params.TextureData->VideoTexture->Release();
-				Params.TextureData->SaveContext();
-				Params.TextureData->MakeCurrent();
+						// @todo: this causes a crash
+						//Params.TextureData->VideoTexture->Release();
+						Params.TextureData->SaveContext();
+						Params.TextureData->MakeCurrent();
 
-				if (Params.TextureData->Image != EGL_NO_IMAGE_KHR)
-				{
-					eglDestroyImageKHR(eglGetCurrentDisplay(), Params.TextureData->Image);
-					Params.TextureData->Image = EGL_NO_IMAGE_KHR;
-				}
+						if (Params.TextureData->Image != EGL_NO_IMAGE_KHR)
+						{
+							eglDestroyImageKHR(eglGetCurrentDisplay(), Params.TextureData->Image);
+							Params.TextureData->Image = EGL_NO_IMAGE_KHR;
+						}
 
-				Params.TextureData->RestoreContext();
-				if (Params.TextureData->PreviousNativeBuffer != 0 && MLHandleIsValid(Params.TextureData->PreviousNativeBuffer))
-				{
-					Params.MediaPlayer->RenderThreadReleaseNativeBuffer(Params.MediaPlayerHandle, Params.TextureData->PreviousNativeBuffer);
-				}
+						Params.TextureData->RestoreContext();
+						if (Params.TextureData->PreviousNativeBuffer != 0 && MLHandleIsValid(Params.TextureData->PreviousNativeBuffer))
+						{
+							Params.MediaPlayer->RenderThreadReleaseNativeBuffer(Params.MediaPlayerHandle, Params.TextureData->PreviousNativeBuffer);
+						}
 
-				MLResult Result = MLMediaPlayerDestroy(Params.MediaPlayerHandle);
-				UE_CLOG(Result != MLResult_Ok, LogMagicLeapMedia, Error, TEXT("MLMediaPlayerDestroy failed with error %s."), UTF8_TO_TCHAR(MLMediaResultGetString(Result)));
+						MLResult Result = MLMediaPlayerDestroy(Params.MediaPlayerHandle);
+						UE_CLOG(Result != MLResult_Ok, LogMagicLeapMedia, Error, TEXT("MLMediaPlayerDestroy failed with error %s."), UTF8_TO_TCHAR(MLMediaResultGetString(Result)));
 					});
 			}
 
@@ -646,10 +646,10 @@ void FMagicLeapMediaPlayer::TickFetch(FTimespan DeltaTime, FTimespan Timecode)
 				MLHandle MediaPlayerHandle;
 			};
 
-			FWriteVideoSampleParams WriteVideoSampleParams = { this, StaticCastSharedPtr<FMagicLeapVideoTextureDataVK>(TextureData), PlayerGuid, MediaPlayerHandle };
+			FWriteVideoSampleParams Params = { this, StaticCastSharedPtr<FMagicLeapVideoTextureDataVK>(TextureData), PlayerGuid, MediaPlayerHandle };
 
-			ENQUEUE_UNIQUE_RENDER_COMMAND_ONEPARAMETER(MagicLeapMediaPlayerWriteVideoSample,
-				FWriteVideoSampleParams, Params, WriteVideoSampleParams,
+			ENQUEUE_RENDER_COMMAND(MagicLeapMediaPlayerWriteVideoSample)(
+				[Params](FRHICommandListImmediate& RHICmdList)
 				{
 					auto TextureDataPtr = Params.TextureData.Pin();
 			
@@ -726,10 +726,10 @@ void FMagicLeapMediaPlayer::TickFetch(FTimespan DeltaTime, FTimespan Timecode)
 				MLHandle MediaPlayerHandle;
 			};
 
-			FWriteVideoSampleParams WriteVideoSampleParams = { this, StaticCastSharedPtr<FMagicLeapVideoTextureDataGL>(TextureData), PlayerGuid, MediaPlayerHandle };
+			FWriteVideoSampleParams Params = { this, StaticCastSharedPtr<FMagicLeapVideoTextureDataGL>(TextureData), PlayerGuid, MediaPlayerHandle };
 
-			ENQUEUE_UNIQUE_RENDER_COMMAND_ONEPARAMETER(MagicLeapMediaPlayerWriteVideoSample,
-				FWriteVideoSampleParams, Params, WriteVideoSampleParams,
+			ENQUEUE_RENDER_COMMAND(MagicLeapMediaPlayerWriteVideoSample)(
+				[Params](FRHICommandListImmediate& RHICmdList)
 				{
 					auto TextureDataPtr = Params.TextureData.Pin();
 
@@ -848,11 +848,10 @@ void FMagicLeapMediaPlayer::TickFetch(FTimespan DeltaTime, FTimespan Timecode)
 			FCriticalSection* CriticalSectionPtr;
 			FMediaWorker* MediaWorkerPtr;
 		}
+		Params = { this, StaticCastSharedPtr<FMagicLeapVideoTextureDataGL>(TextureData), MediaPlayerHandle, Samples, VideoSample, &CriticalSection, MediaWorker };
 
-		WriteVideoSampleParams = { this, StaticCastSharedPtr<FMagicLeapVideoTextureDataGL>(TextureData), MediaPlayerHandle, Samples, VideoSample, &CriticalSection, MediaWorker };
-
-		ENQUEUE_UNIQUE_RENDER_COMMAND_ONEPARAMETER(MagicLeapMediaPlayerWriteVideoSample,
-			FWriteVideoSampleParams, Params, WriteVideoSampleParams,
+		ENQUEUE_RENDER_COMMAND(MagicLeapMediaPlayerWriteVideoSample)(
+			[Params](FRHICommandListImmediate& RHICmdList)
 			{
 				auto PinnedTextureData = Params.TextureData.Pin();
 				auto PinnedSamples = Params.SamplesPtr.Pin();
