Updated ML integration to MLSDK 0.19

diff --git a/Engine/Plugins/Lumin/MagicLeapCamera/Source/Public/MagicLeapCameraComponent.h b/Engine/Plugins/Lumin/MagicLeapCamera/Source/Public/MagicLeapCameraComponent.h
new file mode 100644
--- /dev/null
+++ b/Engine/Plugins/Lumin/MagicLeapCamera/Source/Public/MagicLeapCameraComponent.h
@@ -0,0 +1,86 @@
+// Copyright 1998-2018 Epic Games, Inc. All Rights Reserved.
+
+#pragma once
+
+#include "Engine/Engine.h"
+#include "Components/ActorComponent.h"
+#include "MagicLeapCameraTypes.h"
+#include "MagicLeapCameraComponent.generated.h"
+
+/**
+  The MagicLeapCameraComponent provides access to and maintains state for camera capture functionality.
+  The connection to the device's camera is managed internally.  Users of this component
+  are able to asynchronously capture camera images and footage to file.  Alternatively,
+  a camera image can be captured directly to texture.  The user need only make the relevant
+  asynchronous call and then register the appropriate event handlers for the
+  operation's completion.
+*/
+UCLASS(ClassGroup = MagicLeap, BlueprintType, Blueprintable, EditInlineNew, meta = (BlueprintSpawnableComponent))
+class MAGICLEAPCAMERA_API UMagicLeapCameraComponent : public UActorComponent
+{
+	GENERATED_BODY()
+
+public:
+	/** Notifies the MagicLeapCameraPlugin of a new user. */
+	void BeginPlay() override;
+
+	/** Notifies the MagicLeapCameraPlugin that a user is being destroyed (needed for auto-disconnecting the camera). */
+	void EndPlay(const EEndPlayReason::Type EndPlayReason) override;
+
+	/**
+		Initiates a capture image to file task on a separate thread.
+		@brief The newly created jpeg file will have an automatically generated name which is guaranteed
+			   to be unique.  Upon completion, a successful operation will provide the file path of the newly
+			   created jpeg to the FCameraCaptureImgToFile event handler.
+		@return True if the call succeeds, false otherwise.
+	*/
+	UFUNCTION(BlueprintCallable, Category = "Camera | MagicLeap")
+	virtual bool CaptureImageToFileAsync();
+
+	/**
+		Initiates a capture image to memory task on a speparate thread.
+		@brief The user should register event handlers for both the success and fail events.  Upon completion,
+			   a successful operation will provide a dynamically generated texture containing the captured
+			   image to the FCameraCaptureImgToTextureSuccess event handler.
+		@note The generated texture will be garbage collected when this app is destroyed.
+		@return True if the call succeeds, false otherwise.
+	*/
+	UFUNCTION(BlueprintCallable, Category = "Camera | MagicLeap")
+	virtual bool CaptureImageToTextureAsync();
+
+	/**
+		Initiates the capturing of video/audio data on a separate thread.
+		@note The system will continue to record video until StopRecordingVideo is called.
+		@return True if the call succeeds, false otherwise.
+	*/
+	UFUNCTION(BlueprintCallable, Category = "Camera | MagicLeap")
+	virtual bool StartRecordingAsync();
+
+	/**
+		Stops the recording and saves the video/audio data to an mp4 file.
+		@note The newly created mp4 file will have an automatically generated name which is guaranteed
+			  to be unique.
+		@return True if the call succeeds, false otherwise.
+	*/
+	UFUNCTION(BlueprintCallable, Category = "Camera | MagicLeap")
+	virtual bool StopRecordingAsync();
+
+	/**
+		Gets the capture state of the component.
+		@return True if the component is currently capturing, false otherwise.
+	*/
+	UFUNCTION(BlueprintPure, Category = "Camera | MagicLeap")
+	bool IsCapturing() const;
+
+private:
+	UPROPERTY(BlueprintAssignable, Category = "Camera | MagicLeap", meta = (AllowPrivateAccess = true))
+	FCameraCaptureImgToFileMulti OnCaptureImgToFile;
+	UPROPERTY(BlueprintAssignable, Category = "Camera | MagicLeap", meta = (AllowPrivateAccess = true))
+	FCameraCaptureImgToTextureMulti OnCaptureImgToTexture;
+	UPROPERTY(BlueprintAssignable, Category = "Camera | MagicLeap", meta = (AllowPrivateAccess = true))
+	FCameraStartRecordingMulti OnStartRecording;
+	UPROPERTY(BlueprintAssignable, Category = "Camera | MagicLeap", meta = (AllowPrivateAccess = true))
+	FCameraStopRecordingMulti OnStopRecording;
+	UPROPERTY(BlueprintAssignable, Category = "Camera | MagicLeap", meta = (AllowPrivateAccess = true))
+	FCameraLogMessageMulti OnLogMessage;
+};
diff --git a/Engine/Plugins/Lumin/MagicLeapCamera/Source/Public/MagicLeapCameraFunctionLibrary.h b/Engine/Plugins/Lumin/MagicLeapCamera/Source/Public/MagicLeapCameraFunctionLibrary.h
new file mode 100644
--- /dev/null
+++ b/Engine/Plugins/Lumin/MagicLeapCamera/Source/Public/MagicLeapCameraFunctionLibrary.h
@@ -0,0 +1,99 @@
+// Copyright 1998-2018 Epic Games, Inc. All Rights Reserved.
+#pragma once
+
+#include "Engine/Engine.h"
+#include "Kismet/BlueprintFunctionLibrary.h"
+#include "MagicLeapCameraTypes.h"
+#include "MagicLeapCameraFunctionLibrary.generated.h"
+
+/**
+  The MagicLeapCameraFunctionLibrary provides access to the camera capture functionality.
+  Users of this library are able to asynchronously capture camera images and footage to file.
+  Alternatively, a camera image can be captured directly to texture.  The user need only make
+  the relevant asynchronous call and then register the appropriate event handlers for the
+  operation's completion.
+*/
+UCLASS(ClassGroup = MagicLeap)
+class MAGICLEAPCAMERA_API UMagicLeapCameraFunctionLibrary : public UBlueprintFunctionLibrary
+{
+	GENERATED_BODY()
+
+public:
+	/** 
+		Establishes a connection with the device's camera.
+		@note A connection will be made automatically upon the first capture call if this is
+			  not called first.  Calling this function manually allows the developer to control
+			  when privilege notifications for this plugin will be activated (if application
+			  is being used for the first time).
+		@return True if the call succeeds, false otherwise.
+	*/
+	UFUNCTION(BlueprintCallable, Category = "Camera Function Library | MagicLeap")
+	static bool CameraConnect(const FCameraConnect& ResultDelegate);
+
+	/** 
+		Disconnects from the device's camera.
+		@note This function must be called before the application terminates (if the camera has
+			  been connected to).  Failure to do so will result in the camera connection remaining
+			  open (and the camera icon remaining on screen).
+		@return True if the call succeeds, false otherwise.
+	*/
+	UFUNCTION(BlueprintCallable, Category = "Camera Function Library | MagicLeap")
+	static bool CameraDisconnect(const FCameraDisconnect& ResultDelegate);
+
+	/**
+		Initiates a capture image to file task on a separate thread.
+		@brief The newly created jpeg file will have an automatically generated name which is guaranteed
+			   to be unique.  Upon completion, a successful operation will provide the file path of the newly
+			   created jpeg to the FCameraCaptureImgToFile event handler.
+		@param ResultDelegate The delegate to be notified once the camera image has been saved to a jpeg file.
+		@return True if the call succeeds, false otherwise.
+	*/
+	UFUNCTION(BlueprintCallable, Category = "Camera Function Library | MagicLeap")
+	static bool CaptureImageToFileAsync(const FCameraCaptureImgToFile& ResultDelegate);
+
+	/**
+		Initiates a capture image to memory task on a speparate thread.
+		@brief The user should register event handlers for both the success and fail events.  Upon completion,
+			   a successful operation will provide a dynamically generated texture containing the captured
+			   image to the FCameraCaptureImgToTextureSuccess event handler.
+		@note The generated texture will be garbage collected when this app is destroyed.
+		@param ResultDelegate The delegate to be notified once the camera image has been saved to a texture.
+		@return True if the call succeeds, false otherwise.
+	*/
+	UFUNCTION(BlueprintCallable, Category = "Camera Function Library | MagicLeap")
+	static bool CaptureImageToTextureAsync(const FCameraCaptureImgToTexture& ResultDelegate);
+
+	/**
+		Initiates the capturing of video/audio data on a separate thread.
+		@note The system will continue to record video until StopRecordingVideo is called.
+		@param ResultDelegate The delegate to be notified once the recording has begun or failed to begin.
+		@return True if the call succeeds, false otherwise.
+	*/
+	UFUNCTION(BlueprintCallable, Category = "Camera Function Library | MagicLeap")
+	static bool StartRecordingAsync(const FCameraStartRecording& ResultDelegate);
+
+	/**
+		Stops the recording and saves the video/audio data to an mp4 file.
+		@note The newly created mp4 file will have an automatically generated name which is guaranteed
+			  to be unique.
+		@param ResultDelegate The delegate to be notified once the video/audio data has been saved to an mp4 file.
+		@return True if the call succeeds, false otherwise.
+	*/
+	UFUNCTION(BlueprintCallable, Category = "Camera Function Library | MagicLeap")
+	static bool StopRecordingAsync(const FCameraStopRecording& ResultDelegate);
+
+	/**
+		Sets the delegate by which the system can pass log messages back to the calling blueprint.
+		@param LogDelegate The delegate by which the system will return log messages to the calling blueprint.
+		@return True if the call succeeds, false otherwise.
+	*/
+	UFUNCTION(BlueprintCallable, Category = "Camera Function Library | MagicLeap")
+	static bool SetLogDelegate(const FCameraLogMessage& LogDelegate);
+
+	/**
+		Gets the capture state of the component.
+		@return True if the component is currently capturing, false otherwise.
+	*/
+	UFUNCTION(BlueprintPure, Category = "Camera Function Library | MagicLeap")
+	static bool IsCapturing();
+};
diff --git a/Engine/Plugins/Lumin/MagicLeapCamera/Source/Public/MagicLeapCameraTypes.h b/Engine/Plugins/Lumin/MagicLeapCamera/Source/Public/MagicLeapCameraTypes.h
new file mode 100644
--- /dev/null
+++ b/Engine/Plugins/Lumin/MagicLeapCamera/Source/Public/MagicLeapCameraTypes.h
@@ -0,0 +1,57 @@
+// Copyright 1998-2018 Epic Games, Inc. All Rights Reserved.
+#pragma once
+
+#include "Engine/Engine.h"
+#include "MagicLeapCameraTypes.generated.h"
+
+USTRUCT()
+struct FMagicLeapCameraDummyStruct { GENERATED_BODY() };
+
+/** Delegate used to notify the initiating blueprint when the camera connect task has completed. */
+DECLARE_DYNAMIC_DELEGATE_OneParam(FCameraConnect, const bool, bSuccess);
+
+/** Delegate used to notify the initiating blueprint when the camera disonnect task has completed. */
+DECLARE_DYNAMIC_DELEGATE_OneParam(FCameraDisconnect, const bool, bSuccess);
+
+/**
+   Delegate used to notify the initiating blueprint when a capture image to file task has completed.
+   @note Although this signals the task as complete, it may have failed or been cancelled.
+   @param bSuccess True if the task succeeded, false otherwise.
+   @param FilePath A string containing the file path to the newly created jpeg.
+*/
+DECLARE_DYNAMIC_DELEGATE_TwoParams(FCameraCaptureImgToFile, const bool, bSuccess, const FString&, FilePath);
+DECLARE_DYNAMIC_MULTICAST_DELEGATE_TwoParams(FCameraCaptureImgToFileMulti, const bool, bSuccess, const FString&, FilePath);
+
+/**
+	Delegate used to pass the captured image back to the initiating blueprint.
+	@note The captured texture will remain in memory for the lifetime of the calling application (if the task succeeds).
+	@param bSuccess True if the task succeeded, false otherwise.
+	@param CaptureTexture A UTexture2D containing the captured image.
+*/
+DECLARE_DYNAMIC_DELEGATE_TwoParams(FCameraCaptureImgToTexture, const bool, bSuccess, UTexture2D*, CaptureTexture);
+DECLARE_DYNAMIC_MULTICAST_DELEGATE_TwoParams(FCameraCaptureImgToTextureMulti, const bool, bSuccess, UTexture2D*, CaptureTexture);
+
+/**
+	Delegate used to notify the initiating blueprint of the result of a request to begin recording video.
+	@note Although this signals the task as complete, it may have failed or been cancelled.
+	@param bSuccess True if the task succeeded, false otherwise.
+*/
+DECLARE_DYNAMIC_DELEGATE_OneParam(FCameraStartRecording, const bool, bSuccess);
+DECLARE_DYNAMIC_MULTICAST_DELEGATE_OneParam(FCameraStartRecordingMulti, const bool, bSuccess);
+
+/**
+	Delegate used to notify the initiating blueprint of the result of a request to stop recording video.
+	@note Although this signals the task as complete, it may have failed or been cancelled.
+	@param bSuccess True if the task succeeded, false otherwise.
+	@param FilePath A string containing the path to the newly created mp4.
+*/
+DECLARE_DYNAMIC_DELEGATE_TwoParams(FCameraStopRecording, const bool, bSuccess, const FString&, FilePath);
+DECLARE_DYNAMIC_MULTICAST_DELEGATE_TwoParams(FCameraStopRecordingMulti, const bool, bSuccess, const FString&, FilePath);
+
+/**
+	Delegate used to pass log messages from the capture worker thread to the initiating blueprint.
+	@note This is useful if the user wishes to have log messages in 3D space.
+	@param LogMessage A string containing the log message.
+*/
+DECLARE_DYNAMIC_DELEGATE_OneParam(FCameraLogMessage, const FString&, LogMessage);
+DECLARE_DYNAMIC_MULTICAST_DELEGATE_OneParam(FCameraLogMessageMulti, const FString&, LogMessage);
diff --git a/Engine/Plugins/Lumin/MagicLeapMedia/MagicLeapMedia.uplugin b/Engine/Plugins/Lumin/MagicLeapMedia/MagicLeapMedia.uplugin
--- a/Engine/Plugins/Lumin/MagicLeapMedia/MagicLeapMedia.uplugin
+++ b/Engine/Plugins/Lumin/MagicLeapMedia/MagicLeapMedia.uplugin
@@ -25,6 +25,12 @@
 			"LoadingPhase" : "PreLoadingScreen",
 			"WhitelistPlatforms" : [ "Lumin" ]
 		},
+		{
+			"Name" : "MagicLeapMediaCodec",
+			"Type" : "RuntimeNoCommandlet",
+			"LoadingPhase" : "PreLoadingScreen",
+			"WhitelistPlatforms" : [ "Lumin" ]
+		},
 		{
 			"Name" : "MagicLeapMediaEditor",
 			"Type" : "Editor",
@@ -36,6 +42,12 @@
             "LoadingPhase": "PostEngineInit",
             "WhitelistPlatforms": [ "Lumin", "Mac", "Win64" ]
         },
+        {
+			"Name" : "MagicLeapMediaCodecFactory",
+			"Type" : "RuntimeNoCommandlet",
+			"LoadingPhase" : "PostEngineInit",
+			"WhitelistPlatforms" : [ "Lumin", "Mac", "Win64" ]
+		},
         {
 			"Name" : "MagicLeapCameraPreviewFactory",
 			"Type" : "RuntimeNoCommandlet",
diff --git a/Engine/Plugins/Lumin/MagicLeapMedia/Source/MagicLeapCameraPreviewFactory/Private/MagicLeapCameraPreviewFactoryModule.cpp b/Engine/Plugins/Lumin/MagicLeapMedia/Source/MagicLeapCameraPreviewFactory/Private/MagicLeapCameraPreviewFactoryModule.cpp
--- a/Engine/Plugins/Lumin/MagicLeapMedia/Source/MagicLeapCameraPreviewFactory/Private/MagicLeapCameraPreviewFactoryModule.cpp
+++ b/Engine/Plugins/Lumin/MagicLeapMedia/Source/MagicLeapCameraPreviewFactory/Private/MagicLeapCameraPreviewFactoryModule.cpp
@@ -8,6 +8,7 @@
 #include "IMagicLeapMediaModule.h"
 #include "MagicLeapMediaFactoryPrivate.h"
 #include "IMediaPlayerFactory.h"
+#include "Misc/ConfigCacheIni.h"
 
 #define LOCTEXT_NAMESPACE "FMagicLeapCameraPreviewFactoryModule"
 
@@ -27,8 +28,23 @@ public:
 
 	virtual TSharedPtr<IMediaPlayer, ESPMode::ThreadSafe> CreatePlayer(IMediaEventSink& EventSink) override
 	{
-		auto MagicLeapMediaModule = FModuleManager::LoadModulePtr<IMagicLeapMediaModule>("MagicLeapMedia");
-		return (MagicLeapMediaModule != nullptr) ? MagicLeapMediaModule->CreateCameraPreviewPlayer(EventSink) : nullptr;
+		// Disable camera preview on Vulkan until we are able to get it working.
+		// Can't call FLuminPlatformMisc::ShouldUseVulkan from here ...
+		bool bUseVulkan = false;
+		if (GConfig != nullptr)
+		{
+			GConfig->GetBool(TEXT("/Script/LuminRuntimeSettings.LuminRuntimeSettings"), TEXT("bUseVulkan"), bUseVulkan, GEngineIni);
+		}
+		
+		if (!bUseVulkan)
+		{
+			auto MagicLeapMediaModule = FModuleManager::LoadModulePtr<IMagicLeapMediaModule>("MagicLeapMedia");
+			if (MagicLeapMediaModule != nullptr)
+			{
+				return MagicLeapMediaModule->CreateCameraPreviewPlayer(EventSink);
+			}
+		}
+		return nullptr;
 	}
 
 	virtual FText GetDisplayName() const override
diff --git a/Engine/Plugins/Lumin/MagicLeapMedia/Source/MagicLeapMedia/Private/MagicLeapCameraPreviewPlayer.cpp b/Engine/Plugins/Lumin/MagicLeapMedia/Source/MagicLeapMedia/Private/MagicLeapCameraPreviewPlayer.cpp
--- a/Engine/Plugins/Lumin/MagicLeapMedia/Source/MagicLeapMedia/Private/MagicLeapCameraPreviewPlayer.cpp
+++ b/Engine/Plugins/Lumin/MagicLeapMedia/Source/MagicLeapMedia/Private/MagicLeapCameraPreviewPlayer.cpp
@@ -51,6 +51,11 @@ FTimespan FMagicLeapCameraPreviewPlayer::GetTime() const
 	return FTimespan::Zero();
 }
 
+FTimespan FMagicLeapCameraPreviewPlayer::GetDuration() const
+{
+	return FTimespan::Zero();
+}
+
 void FMagicLeapCameraPreviewPlayer::RegisterExternalTexture(const FGuid& InGuid, FTextureRHIRef& InTextureRHI, FSamplerStateRHIRef& InSamplerStateRHI)
 {
 	FExternalTextureRegistry::Get().RegisterExternalTexture(InGuid, InTextureRHI, InSamplerStateRHI, FLinearColor(1.0f, 0.0f, 0.0f, 1.0f), FLinearColor(0.0f, 0.0f, 0.0f, 0.0f));
@@ -77,11 +82,12 @@ bool FMagicLeapCameraPreviewPlayer::RenderThreadIsBufferAvailable(MLHandle Media
 	return UCameraCaptureComponent::GetPreviewHandle() != ML_INVALID_HANDLE;
 }
 
-bool FMagicLeapCameraPreviewPlayer::RenderThreadGetNativeBuffer(const MLHandle MediaPlayerHandle, MLHandle& NativeBuffer)
+bool FMagicLeapCameraPreviewPlayer::RenderThreadGetNativeBuffer(const MLHandle MediaPlayerHandle, MLHandle& NativeBuffer, bool& OutIsVideoTextureValid)
 {
 	ensureMsgf(IsInRenderingThread(), TEXT("RenderThreadGetNativeBuffer called outside of render thread"));
 	(void)MediaPlayerHandle;
 	NativeBuffer = UCameraCaptureComponent::GetPreviewHandle();
+	OutIsVideoTextureValid = true;
 	return NativeBuffer != ML_INVALID_HANDLE;
 }
 
diff --git a/Engine/Plugins/Lumin/MagicLeapMedia/Source/MagicLeapMedia/Private/MagicLeapCameraPreviewPlayer.h b/Engine/Plugins/Lumin/MagicLeapMedia/Source/MagicLeapMedia/Private/MagicLeapCameraPreviewPlayer.h
--- a/Engine/Plugins/Lumin/MagicLeapMedia/Source/MagicLeapMedia/Private/MagicLeapCameraPreviewPlayer.h
+++ b/Engine/Plugins/Lumin/MagicLeapMedia/Source/MagicLeapMedia/Private/MagicLeapCameraPreviewPlayer.h
@@ -21,13 +21,14 @@ public:
 	bool GetVideoTrackFormat(int32 TrackIndex, int32 FormatIndex, FMediaVideoTrackFormat& OutFormat) const override;
 	bool IsLooping() const override;
 	FTimespan GetTime() const override;
+	FTimespan GetDuration() const override;
 
 private:
 	bool SetRateOne() override;
 	bool GetMediaPlayerState(uint16 FlagToPoll) const override;
 	void RegisterExternalTexture(const FGuid& InGuid, FTextureRHIRef& InTextureRHI, FSamplerStateRHIRef& InSamplerStateRHI) override;
 	bool RenderThreadIsBufferAvailable(MLHandle MediaPlayerHandle) override;
-	bool RenderThreadGetNativeBuffer(const MLHandle MediaPlayerHandle, MLHandle& NativeBuffer) override;
+	bool RenderThreadGetNativeBuffer(const MLHandle MediaPlayerHandle, MLHandle& NativeBuffer, bool& OutIsVideoTextureValid) override;
 	bool RenderThreadReleaseNativeBuffer(const MLHandle MediaPlayerHandle, MLHandle NativeBuffer) override;
 	bool RenderThreadGetCurrentPosition(const MLHandle MediaPlayerHandle, int32& CurrentPosition) override;
 };
diff --git a/Engine/Plugins/Lumin/MagicLeapMedia/Source/MagicLeapMedia/Private/MagicLeapMediaPlayer.cpp b/Engine/Plugins/Lumin/MagicLeapMedia/Source/MagicLeapMedia/Private/MagicLeapMediaPlayer.cpp
--- a/Engine/Plugins/Lumin/MagicLeapMedia/Source/MagicLeapMedia/Private/MagicLeapMediaPlayer.cpp
+++ b/Engine/Plugins/Lumin/MagicLeapMedia/Source/MagicLeapMedia/Private/MagicLeapMediaPlayer.cpp
@@ -145,6 +145,7 @@ FMagicLeapMediaPlayer::FMagicLeapMediaPlayer(IMediaEventSink& InEventSink)
 , MediaWorker(nullptr)
 , bWasMediaPlayingBeforeAppPause(false)
 , bPlaybackCompleted(false)
+, CurrentPlaybackTime(FTimespan::Zero())
 {
 	if (FLuminPlatformMisc::ShouldUseVulkan())
 	{
@@ -156,7 +157,7 @@ FMagicLeapMediaPlayer::FMagicLeapMediaPlayer(IMediaEventSink& InEventSink)
 	}
 	
 	MLResult Result = MLMediaPlayerCreate(&MediaPlayerHandle);
-	UE_CLOG(Result != MLResult_Ok, LogMagicLeapMedia, Error, TEXT("MLMediaPlayerCreate failed with error %d."), Result);
+	UE_CLOG(Result != MLResult_Ok, LogMagicLeapMedia, Error, TEXT("MLMediaPlayerCreate failed with error %s."), UTF8_TO_TCHAR(MLMediaResultGetString(Result)));
 	CurrentState = (Samples.IsValid() && Result == MLResult_Ok) ? EMediaState::Closed : EMediaState::Error;
 	
 	if (!GSupportsImageExternal)
@@ -202,7 +203,7 @@ FMagicLeapMediaPlayer::~FMagicLeapMediaPlayer()
 						if (MLHandleIsValid(Params.MediaPlayerHandle))
 						{
 							MLResult Result = MLMediaPlayerDestroy(Params.MediaPlayerHandle);
-							UE_CLOG(Result != MLResult_Ok, LogMagicLeapMedia, Error, TEXT("MLMediaPlayerDestroy failed with error %d."), Result);
+							UE_CLOG(Result != MLResult_Ok, LogMagicLeapMedia, Error, TEXT("MLMediaPlayerDestroy failed with error %s."), UTF8_TO_TCHAR(MLMediaResultGetString(Result)));
 						}
 					});
 			}
@@ -243,7 +244,7 @@ FMagicLeapMediaPlayer::~FMagicLeapMediaPlayer()
 				}
 
 				MLResult Result = MLMediaPlayerDestroy(Params.MediaPlayerHandle);
-				UE_CLOG(Result != MLResult_Ok, LogMagicLeapMedia, Error, TEXT("MLMediaPlayerDestroy failed with error %d."), Result);
+				UE_CLOG(Result != MLResult_Ok, LogMagicLeapMedia, Error, TEXT("MLMediaPlayerDestroy failed with error %s."), UTF8_TO_TCHAR(MLMediaResultGetString(Result)));
 					});
 			}
 
@@ -252,7 +253,7 @@ FMagicLeapMediaPlayer::~FMagicLeapMediaPlayer()
 		else
 		{
 			MLResult Result = MLMediaPlayerDestroy(MediaPlayerHandle);
-			UE_CLOG(Result != MLResult_Ok, LogMagicLeapMedia, Error, TEXT("MLMediaPlayerDestroy failed with error %d."), Result);
+			UE_CLOG(Result != MLResult_Ok, LogMagicLeapMedia, Error, TEXT("MLMediaPlayerDestroy failed with error %s."), UTF8_TO_TCHAR(MLMediaResultGetString(Result)));
 		}
 
 		MediaPlayerHandle = ML_INVALID_HANDLE;
@@ -269,6 +270,11 @@ void FMagicLeapMediaPlayer::Close()
 		return;
 	}
 
+	{
+		FScopeLock Lock(&CriticalSection);
+		bPlaybackCompleted = true;
+	}
+
 	// remove delegates if registered
 	if (ResumeHandle.IsValid())
 	{
@@ -283,10 +289,10 @@ void FMagicLeapMediaPlayer::Close()
 	}
 
 	MLResult StopResult = MLMediaPlayerStop(MediaPlayerHandle);
-	UE_CLOG(StopResult != MLResult_Ok, LogMagicLeapMedia, Error, TEXT("MLMediaPlayerStop failed with error %d"), StopResult);
+	UE_CLOG(StopResult != MLResult_Ok, LogMagicLeapMedia, Error, TEXT("MLMediaPlayerStop failed with error %s"), UTF8_TO_TCHAR(MLMediaResultGetString(StopResult)));
 
 	MLResult ResetResult = MLMediaPlayerReset(MediaPlayerHandle);
-	UE_CLOG(ResetResult != MLResult_Ok, LogMagicLeapMedia, Error, TEXT("MLMediaPlayerReset failed with error %d"), ResetResult);
+	UE_CLOG(ResetResult != MLResult_Ok, LogMagicLeapMedia, Error, TEXT("MLMediaPlayerReset failed with error %s"), UTF8_TO_TCHAR(MLMediaResultGetString(ResetResult)));
 
 	CurrentState = EMediaState::Closed;
 
@@ -388,7 +394,7 @@ bool FMagicLeapMediaPlayer::Open(const FString& Url, const IMediaOptions* Option
 		MLResult Result = MLMediaPlayerSetDataSourceForPath(MediaPlayerHandle, TCHAR_TO_UTF8(*FilePath));
 		if (Result != MLResult_Ok)
 		{
-			UE_LOG(LogMagicLeapMedia, Error, TEXT("MLMediaPlayerSetDataSourceForPath for path %s failed with error %d."), *FilePath, Result);
+			UE_LOG(LogMagicLeapMedia, Error, TEXT("MLMediaPlayerSetDataSourceForPath for path %s failed with error %s."), *FilePath, UTF8_TO_TCHAR(MLMediaResultGetString(Result)));
 			EventSink.ReceiveMediaEvent(EMediaEvent::MediaOpenFailed);
 			return false;
 		}
@@ -399,7 +405,7 @@ bool FMagicLeapMediaPlayer::Open(const FString& Url, const IMediaOptions* Option
 		MLResult Result = MLMediaPlayerSetDataSourceForURI(MediaPlayerHandle, TCHAR_TO_UTF8(*Url));
 		if (Result != MLResult_Ok)
 		{
-			UE_LOG(LogMagicLeapMedia, Error, TEXT("MLMediaPlayerSetDataSourceForURI for remote media source %s failed with error %d."), *Url, Result);
+			UE_LOG(LogMagicLeapMedia, Error, TEXT("MLMediaPlayerSetDataSourceForURI for remote media source %s failed with error %s."), *Url, UTF8_TO_TCHAR(MLMediaResultGetString(Result)));
 			EventSink.ReceiveMediaEvent(EMediaEvent::MediaOpenFailed);
 			return false;
 		}
@@ -413,7 +419,7 @@ bool FMagicLeapMediaPlayer::Open(const FString& Url, const IMediaOptions* Option
 	MLResult Result = MLMediaPlayerPrepareAsync(MediaPlayerHandle);
 	if (Result != MLResult_Ok)
 	{
-		UE_LOG(LogMagicLeapMedia, Error, TEXT("MLMediaPlayerPrepareAsync for media source %s failed with error %d"), *Url, Result);
+		UE_LOG(LogMagicLeapMedia, Error, TEXT("MLMediaPlayerPrepareAsync for media source %s failed with error %s"), *Url, UTF8_TO_TCHAR(MLMediaResultGetString(Result)));
 		EventSink.ReceiveMediaEvent(EMediaEvent::MediaOpenFailed);
 		return false;
 	}
@@ -463,7 +469,7 @@ FTimespan FMagicLeapMediaPlayer::GetDuration() const
 		}
 		else
 		{
-			UE_LOG(LogMagicLeapMedia, Error, TEXT("MLMediaPlayerGetDuration failed with error %d"), Result);
+			UE_LOG(LogMagicLeapMedia, Error, TEXT("MLMediaPlayerGetDuration failed with error %s"), UTF8_TO_TCHAR(MLMediaResultGetString(Result)));
 		}
 	}
 
@@ -497,23 +503,7 @@ TRangeSet<float> FMagicLeapMediaPlayer::GetSupportedRates(EMediaRateThinning Thi
 
 FTimespan FMagicLeapMediaPlayer::GetTime() const
 {
-	FTimespan Time = FTimespan::Zero();
-
-	if (CurrentState == EMediaState::Playing || CurrentState == EMediaState::Paused)
-	{
-		int32 currentPositionMilSec = 0;
-		MLResult Result = MLMediaPlayerGetCurrentPosition(MediaPlayerHandle, &currentPositionMilSec);
-		if (Result == MLResult_Ok)
-		{
-			Time = FTimespan::FromMilliseconds(currentPositionMilSec);
-		}
-		else
-		{
-			UE_LOG(LogMagicLeapMedia, Error, TEXT("MLMediaPlayerGetCurrentPosition failed with error %d"), Result);
-		}
-	}
-
-	return Time;
+	return CurrentPlaybackTime.Load();
 }
 
 bool FMagicLeapMediaPlayer::IsLooping() const
@@ -532,11 +522,11 @@ bool FMagicLeapMediaPlayer::Seek(const FTimespan& Time)
 	}
 	else if (CurrentState == EMediaState::Playing || CurrentState == EMediaState::Paused)
 	{
-		MLResult Result = MLMediaPlayerSeekTo(MediaPlayerHandle, Time.GetTotalMilliseconds(), MLMediaSeekMode_Previous_Sync);
+		MLResult Result = MLMediaPlayerSeekTo(MediaPlayerHandle, Time.GetTotalMilliseconds(), MLMediaSeekMode_Closest_Sync);
 		if (Result != MLResult_Ok)
 		{
 			bSuccess = false;
-			UE_LOG(LogMagicLeapMedia, Error, TEXT("MLMediaPlayerSeekTo failed with error %d"), Result);
+			UE_LOG(LogMagicLeapMedia, Error, TEXT("MLMediaPlayerSeekTo failed with error %s"), UTF8_TO_TCHAR(MLMediaResultGetString(Result)));
 		}
 	}
 
@@ -548,7 +538,7 @@ bool FMagicLeapMediaPlayer::SetLooping(bool Looping)
 	MLResult Result = MLMediaPlayerSetLooping(MediaPlayerHandle, Looping);
 	if (Result != MLResult_Ok)
 	{
-		UE_LOG(LogMagicLeapMedia, Error, TEXT("MLMediaPlayerSetLooping failed with error %d"), Result);
+		UE_LOG(LogMagicLeapMedia, Error, TEXT("MLMediaPlayerSetLooping failed with error %s"), UTF8_TO_TCHAR(MLMediaResultGetString(Result)));
 		return false;
 	}
 
@@ -580,7 +570,7 @@ bool FMagicLeapMediaPlayer::SetRate(float Rate)
 		}
 		else
 		{
-			UE_LOG(LogMagicLeapMedia, Error, TEXT("MLMediaPlayerPause failed with error %d!"), Result);
+			UE_LOG(LogMagicLeapMedia, Error, TEXT("MLMediaPlayerPause failed with error %s!"), UTF8_TO_TCHAR(MLMediaResultGetString(Result)));
 		}
 	}
 	else if (Rate == 1.0f)
@@ -607,7 +597,7 @@ bool FMagicLeapMediaPlayer::SetNativeVolume(float Volume)
 		if (Result != MLResult_Ok)
 		{
 			bSuccess = false;
-			UE_LOG(LogMagicLeapMedia, Error, TEXT("MLMediaPlayerSetVolume failed with error %d."), Result);
+			UE_LOG(LogMagicLeapMedia, Error, TEXT("MLMediaPlayerSetVolume failed with error %s."), UTF8_TO_TCHAR(MLMediaResultGetString(Result)));
 		}
 	}
 
@@ -626,6 +616,17 @@ void FMagicLeapMediaPlayer::TickFetch(FTimespan DeltaTime, FTimespan Timecode)
 		return;
 	}
 
+	int32 currentPositionMilSec = 0;
+	MLResult Result = MLMediaPlayerGetCurrentPosition(MediaPlayerHandle, &currentPositionMilSec);
+	if (Result == MLResult_Ok)
+	{
+		CurrentPlaybackTime = FTimespan::FromMilliseconds(currentPositionMilSec);
+	}
+	else
+	{
+		UE_LOG(LogMagicLeapMedia, Error, TEXT("MLMediaPlayerGetCurrentPosition failed with error %s"), UTF8_TO_TCHAR(MLMediaResultGetString(Result)));
+	}
+
 	// deal with resolution changes (usually from streams)
 	if (GetMediaPlayerState(MLMediaPlayerPollingStateFlag_HasSizeChanged))
 	{
@@ -664,7 +665,7 @@ void FMagicLeapMediaPlayer::TickFetch(FTimespan DeltaTime, FTimespan Timecode)
 					}
 
 					MLHandle NativeBuffer = ML_INVALID_HANDLE;
-					if (!Params.MediaPlayer->RenderThreadGetNativeBuffer(Params.MediaPlayerHandle, NativeBuffer))
+					if (!Params.MediaPlayer->RenderThreadGetNativeBuffer(Params.MediaPlayerHandle, NativeBuffer, TextureDataPtr->bIsVideoTextureValid))
 					{
 						return;
 					}
@@ -708,7 +709,7 @@ void FMagicLeapMediaPlayer::TickFetch(FTimespan DeltaTime, FTimespan Timecode)
 
 					if (!TextureDataPtr->bIsVideoTextureValid)
 					{
-						FExternalTextureRegistry::Get().RegisterExternalTexture(Params.PlayerGuid, TextureDataPtr->VideoTexture, TextureDataPtr->VideoSampler, FLinearColor(1.0f, 0.0f, 0.0f, 1.0f), FLinearColor(0.0f, 0.0f, 0.0f, 0.0f));
+						Params.MediaPlayer->RegisterExternalTexture(Params.PlayerGuid, TextureDataPtr->VideoTexture, TextureDataPtr->VideoSampler);
 						TextureDataPtr->bIsVideoTextureValid = true;
 					}
 
@@ -756,7 +757,7 @@ void FMagicLeapMediaPlayer::TickFetch(FTimespan DeltaTime, FTimespan Timecode)
 					// MLHandle because Unreal's uint64 is 'unsigned long long *' whereas uint64_t for the C-API is 'unsigned long *'
 					// TODO: Fix the Unreal types for the above comment.
 					MLHandle nativeBuffer = ML_INVALID_HANDLE;
-					if (!Params.MediaPlayer->RenderThreadGetNativeBuffer(Params.MediaPlayerHandle, nativeBuffer))
+					if (!Params.MediaPlayer->RenderThreadGetNativeBuffer(Params.MediaPlayerHandle, nativeBuffer, TextureDataPtr->bIsVideoTextureValid))
 					{
 						return;
 					}
@@ -814,8 +815,7 @@ void FMagicLeapMediaPlayer::TickFetch(FTimespan DeltaTime, FTimespan Timecode)
 					{
 						FSamplerStateInitializerRHI SamplerStateInitializer(SF_Bilinear, AM_Clamp, AM_Clamp, AM_Clamp);
 						FSamplerStateRHIRef SamplerStateRHI = RHICreateSamplerState(SamplerStateInitializer);
-						FExternalTextureRegistry::Get().RegisterExternalTexture(Params.PlayerGuid, MediaVideoTexture, SamplerStateRHI, FLinearColor(1.0f, 0.0f, 0.0f, 1.0f), FLinearColor(0.0f, 0.0f, 0.0f, 0.0f));
-
+						Params.MediaPlayer->RegisterExternalTexture(Params.PlayerGuid, MediaVideoTexture, SamplerStateRHI);
 						TextureDataPtr->bIsVideoTextureValid = true;
 					}
 				});
@@ -895,20 +895,14 @@ void FMagicLeapMediaPlayer::TickFetch(FTimespan DeltaTime, FTimespan Timecode)
 					// MLHandle because Unreal's uint64 is 'unsigned long long *' whereas uint64_t for the C-API is 'unsigned long *'
 					// TODO: Fix the Unreal types for the above comment.
 					MLHandle NativeBuffer = ML_INVALID_HANDLE;
-					if (!Params.MediaPlayer->RenderThreadGetNativeBuffer(Params.MediaPlayerHandle, NativeBuffer))
+					if (!Params.MediaPlayer->RenderThreadGetNativeBuffer(Params.MediaPlayerHandle, NativeBuffer, PinnedTextureData->bIsVideoTextureValid))
 					{
 						UE_LOG(LogMagicLeapMedia, Error, TEXT("Error acquiring next video buffer"));
 						return;
 					}
 
-					int32 currentPositionMilSec = 0;
-					if (!Params.MediaPlayer->RenderThreadGetCurrentPosition(Params.MediaPlayerHandle, currentPositionMilSec))
-					{
-						UE_LOG(LogMagicLeapMedia, Error, TEXT("Error acquiring current position."));
-						return;
-					}
 					// write frame into texture
-					FRHITexture2D* Texture = Params.VideoSample->InitializeTexture(FTimespan::FromMilliseconds(currentPositionMilSec));
+					FRHITexture2D* Texture = Params.VideoSample->InitializeTexture(Params.MediaPlayer->GetTime());
 					if (Texture != nullptr)
 					{
 						int32 Resource = *reinterpret_cast<int32*>(Texture->GetNativeResource());
@@ -1103,7 +1097,7 @@ FString FMagicLeapMediaPlayer::GetTrackLanguage(EMediaTrackType TrackType, int32
 		}
 		else
 		{
-			UE_LOG(LogMagicLeapMedia, Error, TEXT("MLMediaPlayerGetTrackLanguage failed with error %d"), Result);
+			UE_LOG(LogMagicLeapMedia, Error, TEXT("MLMediaPlayerGetTrackLanguage failed with error %s"), UTF8_TO_TCHAR(MLMediaResultGetString(Result)));
 		}
 	}
 
@@ -1128,7 +1122,7 @@ bool FMagicLeapMediaPlayer::GetVideoTrackFormat(int32 TrackIndex, int32 FormatIn
 	MLResult Result = MLMediaPlayerGetVideoSize(MediaPlayerHandle, &width, &height);
 	if (Result != MLResult_Ok)
 	{
-		UE_LOG(LogMagicLeapMedia, Error, TEXT("MLMediaPlayerGetVideoSize failed with error %d"), Result);
+		UE_LOG(LogMagicLeapMedia, Error, TEXT("MLMediaPlayerGetVideoSize failed with error %s"), UTF8_TO_TCHAR(MLMediaResultGetString(Result)));
 		return false;
 	}
 
@@ -1154,7 +1148,7 @@ bool FMagicLeapMediaPlayer::SelectTrack(EMediaTrackType TrackType, int32 TrackIn
 			}
 			else
 			{
-				UE_LOG(LogMagicLeapMedia, Error, TEXT("MLMediaPlayerSelectTrack failed with error %d"), Result);
+				UE_LOG(LogMagicLeapMedia, Error, TEXT("MLMediaPlayerSelectTrack failed with error %s"), UTF8_TO_TCHAR(MLMediaResultGetString(Result)));
 			}
 		}
 	}
@@ -1194,7 +1188,7 @@ FIntPoint FMagicLeapMediaPlayer::GetVideoDimensions() const
 	MLResult Result = MLMediaPlayerGetVideoSize(MediaPlayerHandle, &width, &height);
 	if (Result != MLResult_Ok)
 	{
-		UE_LOG(LogMagicLeapMedia, Error, TEXT("MLMediaPlayerGetVideoSize failed with error %d"), Result);
+		UE_LOG(LogMagicLeapMedia, Error, TEXT("MLMediaPlayerGetVideoSize failed with error %s"), UTF8_TO_TCHAR(MLMediaResultGetString(Result)));
 	}
 	return FIntPoint(width, height);
 }
@@ -1204,7 +1198,7 @@ bool FMagicLeapMediaPlayer::SetRateOne()
 	MLResult Result = MLMediaPlayerStart(MediaPlayerHandle);
 	if (Result != MLResult_Ok)
 	{
-		UE_LOG(LogMagicLeapMedia, Error, TEXT("MLMediaPlayerStart failed with error %d"), Result);
+		UE_LOG(LogMagicLeapMedia, Error, TEXT("MLMediaPlayerStart failed with error %s"), UTF8_TO_TCHAR(MLMediaResultGetString(Result)));
 		return false;
 	}
 	
@@ -1219,7 +1213,7 @@ bool FMagicLeapMediaPlayer::GetMediaPlayerState(uint16 FlagToPoll) const
 	MLResult Result = MLMediaPlayerPollStates(MediaPlayerHandle, FlagToPoll, &StateFlags);
 	if (Result != MLResult_Ok)
 	{
-		UE_LOG(LogMagicLeapMedia, Error, TEXT("MLMediaPlayerPollStates failed with error %d"), Result);
+		UE_LOG(LogMagicLeapMedia, Error, TEXT("MLMediaPlayerPollStates failed with error %s"), UTF8_TO_TCHAR(MLMediaResultGetString(Result)));
 		return false;
 	}
 
@@ -1233,23 +1227,42 @@ bool FMagicLeapMediaPlayer::RenderThreadIsBufferAvailable(MLHandle InMediaPlayer
 	MLResult Result = MLMediaPlayerPollStates(InMediaPlayerHandle, MLMediaPlayerPollingStateFlag_IsBufferAvailable, &StateFlags);
 	if (Result != MLResult_Ok)
 	{
-		UE_LOG(LogMagicLeapMedia, Error, TEXT("MLMediaPlayerPollStates failed with error %d"), Result);
+		UE_LOG(LogMagicLeapMedia, Error, TEXT("MLMediaPlayerPollStates failed with error %s"), UTF8_TO_TCHAR(MLMediaResultGetString(Result)));
 		return false;
 	}
 
 	return MLMediaPlayerPollingStateFlag_IsBufferAvailable & StateFlags;
 }
 
-bool FMagicLeapMediaPlayer::RenderThreadGetNativeBuffer(const MLHandle InMediaPlayerHandle, MLHandle& NativeBuffer)
+bool FMagicLeapMediaPlayer::RenderThreadGetNativeBuffer(const MLHandle InMediaPlayerHandle, MLHandle& NativeBuffer, bool& OutIsVideoTextureValid)
 {
 	ensureMsgf(IsInRenderingThread(), TEXT("RenderThreadGetNativeBuffer called outside of render thread"));
 	MLResult Result = MLMediaPlayerAcquireNextAvailableBuffer(InMediaPlayerHandle, &NativeBuffer);
 	if (Result != MLResult_Ok)
 	{
-		UE_LOG(LogMagicLeapMedia, Error, TEXT("MLMediaPlayerAcquireNextAvailableBuffer failed with error %d"), Result);
+		UE_LOG(LogMagicLeapMedia, Error, TEXT("MLMediaPlayerAcquireNextAvailableBuffer failed with error %s"), UTF8_TO_TCHAR(MLMediaResultGetString(Result)));
 		return false;
 	}
 
+	Result = MLMediaPlayerGetFrameTransformationMatrix(InMediaPlayerHandle, FrameTransformationMatrix);
+	if (Result != MLResult_Ok)
+	{
+		UE_LOG(LogMagicLeapMedia, Error, TEXT("MLMediaPlayerGetFrameTransformationMatrix failed with error %s"), UTF8_TO_TCHAR(MLMediaResultGetString(Result)));
+		return false;
+	}
+
+	if (UScale != FrameTransformationMatrix[0] ||
+		UOffset != FrameTransformationMatrix[12] ||
+		(-VScale) != FrameTransformationMatrix[5] /*||
+		(1.0f - VOffset) != FrameTransformationMatrix[13]*/)
+	{
+		UScale = FrameTransformationMatrix[0];
+		UOffset = FrameTransformationMatrix[12];
+		VScale = -FrameTransformationMatrix[5];
+		//VOffset = 1.0f - FrameTransformationMatrix[13];
+		OutIsVideoTextureValid = false;
+	}
+
 	return true;
 }
 
@@ -1259,7 +1272,7 @@ bool FMagicLeapMediaPlayer::RenderThreadReleaseNativeBuffer(const MLHandle InMed
 	MLResult Result = MLMediaPlayerReleaseBuffer(InMediaPlayerHandle, NativeBuffer);
 	if (Result != MLResult_Ok)
 	{
-		UE_LOG(LogMagicLeapMedia, Error, TEXT("MLMediaPlayerReleaseBuffer failed with error %d"), Result);
+		UE_LOG(LogMagicLeapMedia, Error, TEXT("MLMediaPlayerReleaseBuffer failed with error %s"), UTF8_TO_TCHAR(MLMediaResultGetString(Result)));
 		return false;
 	}
 
@@ -1272,7 +1285,7 @@ bool FMagicLeapMediaPlayer::RenderThreadGetCurrentPosition(const MLHandle InMedi
 	MLResult Result = MLMediaPlayerGetCurrentPosition(InMediaPlayerHandle, &CurrentPosition);
 	if (Result != MLResult_Ok)
 	{
-		UE_LOG(LogMagicLeapMedia, Error, TEXT("MLMediaPlayerGetCurrentPosition failed with error %d"), Result);
+		UE_LOG(LogMagicLeapMedia, Error, TEXT("MLMediaPlayerGetCurrentPosition failed with error %s"), UTF8_TO_TCHAR(MLMediaResultGetString(Result)));
 		return false;
 	}
 
diff --git a/Engine/Plugins/Lumin/MagicLeapMedia/Source/MagicLeapMedia/Private/MagicLeapMediaPlayer.h b/Engine/Plugins/Lumin/MagicLeapMedia/Source/MagicLeapMedia/Private/MagicLeapMediaPlayer.h
--- a/Engine/Plugins/Lumin/MagicLeapMedia/Source/MagicLeapMedia/Private/MagicLeapMediaPlayer.h
+++ b/Engine/Plugins/Lumin/MagicLeapMedia/Source/MagicLeapMedia/Private/MagicLeapMediaPlayer.h
@@ -12,6 +12,7 @@
 #include "IMediaView.h"
 #include "Misc/Timespan.h"
 #include "Templates/SharedPointer.h"
+#include "Templates/Atomic.h"
 #include "RenderingThread.h"
 #include "RHI.h"
 #include "RHIResources.h"
@@ -153,12 +154,14 @@ protected:
 
 	bool bPlaybackCompleted;
 
+	TAtomic<FTimespan> CurrentPlaybackTime;
+
 private:
 	virtual bool SetRateOne();
 	virtual bool GetMediaPlayerState(uint16 FlagToPoll) const;
 	virtual void RegisterExternalTexture(const FGuid& InGuid, FTextureRHIRef& InTextureRHI, FSamplerStateRHIRef& InSamplerStateRHI);
 	virtual bool RenderThreadIsBufferAvailable(MLHandle MediaPlayerHandle);
-	virtual bool RenderThreadGetNativeBuffer(const MLHandle MediaPlayerHandle, MLHandle& NativeBuffer);
+	virtual bool RenderThreadGetNativeBuffer(const MLHandle MediaPlayerHandle, MLHandle& NativeBuffer, bool& OutIsVideoTextureValid);
 	virtual bool RenderThreadReleaseNativeBuffer(const MLHandle MediaPlayerHandle, MLHandle NativeBuffer);
 	virtual bool RenderThreadGetCurrentPosition(const MLHandle MediaPlayerHandle, int32& CurrentPosition);
 };
diff --git a/Engine/Plugins/Lumin/MagicLeapMedia/Source/MagicLeapMedia/Private/MediaWorker.cpp b/Engine/Plugins/Lumin/MagicLeapMedia/Source/MagicLeapMedia/Private/MediaWorker.cpp
--- a/Engine/Plugins/Lumin/MagicLeapMedia/Source/MagicLeapMedia/Private/MediaWorker.cpp
+++ b/Engine/Plugins/Lumin/MagicLeapMedia/Source/MagicLeapMedia/Private/MediaWorker.cpp
@@ -61,6 +61,7 @@ uint32 FMediaWorker::Run()
 	{
 		uint16_t States = 0;
 		MLResult Result = MLMediaPlayerPollStates(MediaPlayerHandle, MLMediaPlayerPollingStateFlag_IsBufferAvailable, &States);
+		UE_CLOG(Result != MLResult_Ok, LogMagicLeapMedia, Error, TEXT("MLMediaPlayerPollStates failed with error %s."), UTF8_TO_TCHAR(MLMediaResultGetString(Result)));
 		if (Result != MLResult_Ok || !(MLMediaPlayerPollingStateFlag_IsBufferAvailable & States))
 		{
 			continue;
@@ -68,6 +69,7 @@ uint32 FMediaWorker::Run()
 
 		int32 currentPositionMilSec = 0;
 		Result = MLMediaPlayerGetCurrentPosition(MediaPlayerHandle, &currentPositionMilSec);
+		UE_CLOG(Result != MLResult_Ok, LogMagicLeapMedia, Error, TEXT("MLMediaPlayerGetCurrentPosition failed with error %s."), UTF8_TO_TCHAR(MLMediaResultGetString(Result)));
 		const FTimespan Time = FTimespan::FromMilliseconds(currentPositionMilSec);
 
 		int32 width = 0;
@@ -75,7 +77,7 @@ uint32 FMediaWorker::Run()
 		Result = MLMediaPlayerGetVideoSize(MediaPlayerHandle, &width, &height);
 		if (Result != MLResult_Ok)
 		{
-			UE_LOG(LogMagicLeapMedia, Error, TEXT("Error getting video dimensions"));
+			UE_LOG(LogMagicLeapMedia, Error, TEXT("MLMediaPlayerGetVideoSize failed with error %s."), UTF8_TO_TCHAR(MLMediaResultGetString(Result)));
 			continue;
 		}
 
@@ -101,7 +103,7 @@ uint32 FMediaWorker::Run()
 
 		if (Result != MLResult_Ok)
 		{
-			UE_LOG(LogMagicLeapMedia, Error, TEXT("Error acquiring next video buffer"));
+			UE_LOG(LogMagicLeapMedia, Error, TEXT("MLMediaPlayerAcquireNextAvailableBuffer failed with error %s."), UTF8_TO_TCHAR(MLMediaResultGetString(Result)));
 			continue;
 		}
 
@@ -111,8 +113,8 @@ uint32 FMediaWorker::Run()
 			UE_LOG(LogMagicLeapMedia, Error, TEXT("CopyFrameTexture failed"));
 		}
 
-		MLMediaPlayerReleaseBuffer(MediaPlayerHandle, nativeBuffer);
-
+		Result = MLMediaPlayerReleaseBuffer(MediaPlayerHandle, nativeBuffer);
+		UE_CLOG(Result != MLResult_Ok, LogMagicLeapMedia, Error, TEXT("MLMediaPlayerReleaseBuffer failed with error %s."), UTF8_TO_TCHAR(MLMediaResultGetString(Result)));
 		{
 			FScopeLock Lock(CriticalSection);
 			SwapBuffers();
diff --git a/Engine/Plugins/Lumin/MagicLeapMedia/Source/MagicLeapMediaCodec/MagicLeapMediaCodec.Build.cs b/Engine/Plugins/Lumin/MagicLeapMedia/Source/MagicLeapMediaCodec/MagicLeapMediaCodec.Build.cs
new file mode 100644
--- /dev/null
+++ b/Engine/Plugins/Lumin/MagicLeapMedia/Source/MagicLeapMediaCodec/MagicLeapMediaCodec.Build.cs
@@ -0,0 +1,46 @@
+// Copyright 1998-2018 Epic Games, Inc. All Rights Reserved.
+
+using System.IO;
+
+namespace UnrealBuildTool.Rules
+{
+	public class MagicLeapMediaCodec : ModuleRules
+	{
+		public MagicLeapMediaCodec(ReadOnlyTargetRules Target) : base(Target)
+		{
+            PublicIncludePaths.Add(Path.Combine(ModuleDirectory, "Public"));
+			string EngineDir = Path.GetFullPath(Target.RelativeEnginePath);
+
+			DynamicallyLoadedModuleNames.AddRange(
+				new string[] {
+					"Media"
+			});
+
+			PrivateDependencyModuleNames.AddRange(
+				new string[] {
+					"Core",
+					"Engine",
+					"RenderCore",
+					"MediaUtils",
+					"OpenGLDrv",
+					"LuminRuntimeSettings",
+					"MLSDK",
+                    "MagicLeap",
+                    "MagicLeapHelperVulkan",
+            });
+
+			PrivateIncludePathModuleNames.AddRange(
+				new string[] {
+					"Media"
+			});
+
+			PrivateIncludePaths.AddRange(
+				new string[] {
+                    "MagicLeapMedia/Private",
+					Path.Combine(EngineDir, "Source/Runtime/OpenGLDrv/Private")
+            });
+
+			AddEngineThirdPartyPrivateStaticDependencies(Target, "OpenGL");
+		}
+	}
+}
diff --git a/Engine/Plugins/Lumin/MagicLeapMedia/Source/MagicLeapMediaCodec/Private/MagicLeapMediaAudioSample.h b/Engine/Plugins/Lumin/MagicLeapMedia/Source/MagicLeapMediaCodec/Private/MagicLeapMediaAudioSample.h
new file mode 100644
--- /dev/null
+++ b/Engine/Plugins/Lumin/MagicLeapMedia/Source/MagicLeapMediaCodec/Private/MagicLeapMediaAudioSample.h
@@ -0,0 +1,124 @@
+// Copyright 1998-2018 Epic Games, Inc. All Rights Reserved.
+
+#pragma once
+
+#include "CoreTypes.h"
+#include "Containers/Array.h"
+#include "IMediaAudioSample.h"
+#include "MediaObjectPool.h"
+#include "MediaSampleQueue.h"
+#include "Math/IntPoint.h"
+#include "Misc/Timespan.h"
+
+/**
+ * Implements a media audio sample for MagicLeapMedia.
+ */
+class FMagicLeapMediaAudioSample : public IMediaAudioSample, public IMediaPoolable
+{
+public:
+
+	/** Default constructor. */
+	FMagicLeapMediaAudioSample()
+	: Channels(0)
+	, Duration(FTimespan::Zero())
+	, SampleRate(0)
+	, Time(FTimespan::Zero())
+	{ }
+
+	/** Virtual destructor. */
+	virtual ~FMagicLeapMediaAudioSample() { }
+
+public:
+
+	/**
+	 * Initialize the sample.
+	 *
+	 * @param InBuffer The sample's data buffer.
+	 * @param InSize The size of the sample buffer (in bytes).
+	 * @param InTime The sample time (relative to presentation clock).
+	 * @param InDuration The duration for which the sample is valid.
+	 */
+	bool Initialize(
+		const uint8* InBuffer,
+		uint32 InSize,
+		uint32 InChannels,
+		uint32 InSampleRate,
+		FTimespan InTime,
+		FTimespan InDuration)
+	{
+		if ((InBuffer == nullptr) || (InSize == 0))
+		{
+			return false;
+		}
+
+		Buffer.Reset(InSize);
+		Buffer.Append(InBuffer, InSize);
+
+		Channels = InChannels;
+		Duration = InDuration;
+		SampleRate = InSampleRate;
+		Time = InTime;
+
+		return true;
+	}
+
+public:
+
+	//~ IMediaAudioSample interface
+
+	virtual const void* GetBuffer() override
+	{
+		return Buffer.GetData();
+	}
+
+	virtual uint32 GetChannels() const override
+	{
+		return Channels;
+	}
+
+	virtual FTimespan GetDuration() const override
+	{
+		return Duration;
+	}
+
+	virtual EMediaAudioSampleFormat GetFormat() const override
+	{
+		return EMediaAudioSampleFormat::Int16;
+	}
+
+	virtual uint32 GetFrames() const override
+	{
+		return Buffer.Num() / (Channels * sizeof(int16));
+	}
+
+	virtual uint32 GetSampleRate() const override
+	{
+		return SampleRate;
+	}
+
+	virtual FTimespan GetTime() const override
+	{
+		return Time;
+	}
+
+private:
+
+	/** The sample's data buffer. */
+	TArray<uint8> Buffer;
+
+	/** Number of audio channels. */
+	uint32 Channels;
+
+	/** The duration for which the sample is valid. */
+	FTimespan Duration;
+
+	/** Audio sample rate (in samples per second). */
+	uint32 SampleRate;
+
+	/** Presentation time for which the sample was generated. */
+	FTimespan Time;
+};
+
+
+/** Implements a pool for WMF audio sample objects. */
+class FMagicLeapMediaAudioSamplePool : public TMediaObjectPool<FMagicLeapMediaAudioSample> { };
diff --git a/Engine/Plugins/Lumin/MagicLeapMedia/Source/MagicLeapMediaCodec/Private/MagicLeapMediaCodecModule.cpp b/Engine/Plugins/Lumin/MagicLeapMedia/Source/MagicLeapMediaCodec/Private/MagicLeapMediaCodecModule.cpp
new file mode 100644
--- /dev/null
+++ b/Engine/Plugins/Lumin/MagicLeapMedia/Source/MagicLeapMediaCodec/Private/MagicLeapMediaCodecModule.cpp
@@ -0,0 +1,29 @@
+// Copyright 1998-2018 Epic Games, Inc. All Rights Reserved.
+
+#include "CoreMinimal.h"
+#include "MagicLeapMediaCodecPlayer.h"
+#include "IMagicLeapMediaCodecModule.h"
+#include "Modules/ModuleManager.h"
+
+#define LOCTEXT_NAMESPACE "FMagicLeapMediaCodecModule"
+
+class FMagicLeapMediaCodecModule : public IMagicLeapMediaCodecModule
+{
+public:
+	/** IMagicLeapMediaCodecModule interface */
+	virtual TSharedPtr<IMediaPlayer, ESPMode::ThreadSafe> CreatePlayer(IMediaEventSink& EventSink)
+	{
+		return MakeShareable(new FMagicLeapMediaCodecPlayer(EventSink));
+	}
+
+public:
+	/** IModuleInterface interface */
+	virtual void StartupModule() override { }
+	virtual void ShutdownModule() override { }
+};
+
+IMPLEMENT_MODULE(FMagicLeapMediaCodecModule, MagicLeapMediaCodec)
+
+DEFINE_LOG_CATEGORY(LogMagicLeapMediaCodec);
+
+#undef LOCTEXT_NAMESPACE
