Updated ML integration to MLSDK 0.19

diff --git a/Engine/Plugins/Lumin/MagicLeapMedia/Source/MagicLeapMediaCodec/Private/MagicLeapMediaCodecPlayer.cpp b/Engine/Plugins/Lumin/MagicLeapMedia/Source/MagicLeapMediaCodec/Private/MagicLeapMediaCodecPlayer.cpp
new file mode 100644
--- /dev/null
+++ b/Engine/Plugins/Lumin/MagicLeapMedia/Source/MagicLeapMediaCodec/Private/MagicLeapMediaCodecPlayer.cpp
@@ -0,0 +1,1748 @@
+// Copyright 1998-2018 Epic Games, Inc. All Rights Reserved.
+
+#include "MagicLeapMediaCodecPlayer.h"
+#include "IMagicLeapMediaCodecModule.h"
+#include "Misc/Paths.h"
+#include "Lumin/LuminPlatformFile.h"
+#include "Misc/FileHelper.h"
+#include "MediaSamples.h"
+#include "Misc/CoreDelegates.h"
+#include "Misc/ScopeTryLock.h"
+#include "HAL/PlatformAtomics.h"
+#include "UObject/Class.h"
+#include "UObject/UObjectGlobals.h"
+#include "ExternalTexture.h"
+#include "MagicLeapMediaAudioSample.h"
+#include "Lumin/LuminEGL.h"
+#include "MagicLeapHelperVulkan.h"
+
+#if WITH_MLSDK
+#include "ml_version.h"
+#include "ml_media_extractor.h"
+#include "ml_media_codec.h"
+#include "ml_media_codeclist.h"
+#include "ml_media_format.h"
+#include "ml_media_error.h"
+#endif // WITH_MLSDK
+
+#if MLSDK_VERSION_MAJOR == 0 && MLSDK_VERSION_MINOR < 18
+MLMediaFormatKey MLMediaFormat_Key_Mime = "mime";
+MLMediaFormatKey MLMediaFormat_Key_Frame_Rate = "frame-rate";
+MLMediaFormatKey MLMediaFormat_Key_Width = "width";
+MLMediaFormatKey MLMediaFormat_Key_Height = "height";
+MLMediaFormatKey MLMediaFormat_Key_Duration = "durationUs";
+MLMediaFormatKey MLMediaFormat_Key_Language = "language";
+MLMediaFormatKey MLMediaFormat_Key_Sample_Rate = "sample-rate";
+MLMediaFormatKey MLMediaFormat_Key_Channel_Count = "channel-count";
+#endif // MLSDK_VERSION_MAJOR == 0 && MLSDK_VERSION_MINOR < 18
+
+#define LOCTEXT_NAMESPACE "FMagicLeapMediaCodecModule"
+
+/** TData source callback setup */
+int64_t MediaDataSourceReadAt_Callback(MLHandle media_data_source, size_t position, size_t size, uint8_t *buffer, void *context)
+{
+	if (context != nullptr)
+	{
+		return reinterpret_cast<FMagicLeapMediaCodecPlayer*>(context)->MediaDataSourceReadAt(media_data_source, position, size, buffer);
+	}
+
+	return -1;
+}
+
+int64_t MediaDataSourceGetSize_Callback(MLHandle media_data_source, void *context)
+{
+	if (context != nullptr)
+	{
+		return reinterpret_cast<FMagicLeapMediaCodecPlayer*>(context)->MediaDataSourceGetSize(media_data_source);
+	}
+
+	return -1;
+}
+
+void MediaDataSourceClose_Callback(MLHandle media_data_source, void *context)
+{
+	if (context != nullptr)
+	{
+		reinterpret_cast<FMagicLeapMediaCodecPlayer*>(context)->MediaDataSourceClose(media_data_source);
+	}
+}
+
+/** Texture data setup */
+struct FMagicLeapVideoTextureData
+{
+public:
+	FMagicLeapVideoTextureData()
+		: bIsVideoTextureValid(false)
+		, PreviousNativeBuffer(ML_INVALID_HANDLE)
+	{}
+
+	FTextureRHIRef VideoTexture;
+	bool bIsVideoTextureValid;
+	MLHandle PreviousNativeBuffer;
+};
+
+struct FMagicLeapVideoTextureDataVK : public FMagicLeapVideoTextureData
+{
+	FSamplerStateRHIRef VideoSampler;
+	TMap<uint64, FTextureRHIRef> VideoTexturePool;
+};
+
+struct FMagicLeapVideoTextureDataGL : public FMagicLeapVideoTextureData
+{
+public:
+
+	FMagicLeapVideoTextureDataGL()
+		: Image(EGL_NO_IMAGE_KHR)
+		, Display(EGL_NO_DISPLAY)
+		, Context(EGL_NO_CONTEXT)
+		, SavedDisplay(EGL_NO_DISPLAY)
+		, SavedContext(EGL_NO_CONTEXT)
+		, bContextCreated(false)
+	{}
+
+	~FMagicLeapVideoTextureDataGL()
+	{
+		PreviousNativeBuffer = ML_INVALID_HANDLE;
+		eglDestroyContext(Display, Context);
+		Display = EGL_NO_DISPLAY;
+		Context = EGL_NO_CONTEXT;
+	}
+
+	bool InitContext()
+	{
+#if !PLATFORM_LUMINGL4
+		if (Context == EGL_NO_CONTEXT)
+		{
+			Display = LuminEGL::GetInstance()->GetDisplay();
+			EGLContext SharedContext = LuminEGL::GetInstance()->GetCurrentContext();
+			Context = SharedContext; // LuminEGL::GetInstance()->CreateContext(SharedContext);
+		}
+
+		return (Context != EGL_NO_CONTEXT);
+#else
+		return false;
+#endif
+	}
+
+	void SaveContext()
+	{
+#if !PLATFORM_LUMINGL4
+		SavedDisplay = LuminEGL::GetInstance()->GetDisplay();
+		SavedContext = LuminEGL::GetInstance()->GetCurrentContext();
+#endif
+	}
+
+	void MakeCurrent()
+	{
+#if !PLATFORM_LUMINGL4
+		return;	// skip for now
+		EGLBoolean bResult = eglMakeCurrent(Display, EGL_NO_SURFACE, EGL_NO_SURFACE, Context);
+		if (bResult == EGL_FALSE)
+		{
+			UE_LOG(LogMagicLeapMediaCodec, Error, TEXT("Error setting media player context."));
+		}
+#endif
+	}
+
+	void RestoreContext()
+	{
+#if !PLATFORM_LUMINGL4
+		return;	// skip for now
+		EGLBoolean bResult = eglMakeCurrent(SavedDisplay, EGL_NO_SURFACE, EGL_NO_SURFACE, SavedContext);
+		if (bResult == EGL_FALSE)
+		{
+			UE_LOG(LogMagicLeapMediaCodec, Error, TEXT("Error setting unreal context."));
+		}
+#endif
+	}
+
+	EGLImageKHR Image;
+	EGLDisplay Display;
+	EGLContext Context;
+	EGLDisplay SavedDisplay;
+	EGLContext SavedContext;
+	bool bContextCreated;
+};
+
+FMagicLeapMediaCodecPlayer::FMagicLeapMediaCodecPlayer(IMediaEventSink& InEventSink)
+	: AudioCodecHandle(ML_INVALID_HANDLE)
+	, VideoCodecHandle(ML_INVALID_HANDLE)
+	, MediaExtractorHandle(ML_INVALID_HANDLE)
+	, MediaDataSourceHandle(ML_INVALID_HANDLE)
+	, bMediaPrepared(false)
+	, CurrentState(EMediaState::Closed)
+	, EventSink(InEventSink)
+	, Samples(MakeShared<FMediaSamples, ESPMode::ThreadSafe>())
+	, AudioSamplePool(new FMagicLeapMediaAudioSamplePool())
+	, bWasMediaPlayingBeforeAppPause(false)
+	, CurrentPlaybackTime(FTimespan::Zero())
+	, bLoopPlayback(false)
+	, bPlaybackCompleted(false)
+	, bPlaybackCompleted_RenderThread(false)
+	, bIsBufferAvailable(false)
+	, bResetVideoCodecStartTime(false)
+{
+	if (FLuminPlatformMisc::ShouldUseVulkan())
+	{
+		TextureData = MakeShared<FMagicLeapVideoTextureDataVK, ESPMode::ThreadSafe>();
+	}
+	else
+	{
+		TextureData = MakeShared<FMagicLeapVideoTextureDataGL, ESPMode::ThreadSafe>();
+	}
+}
+
+FMagicLeapMediaCodecPlayer::~FMagicLeapMediaCodecPlayer()
+{
+	Close();
+
+	delete AudioSamplePool;
+	AudioSamplePool = nullptr;
+}
+
+void FMagicLeapMediaCodecPlayer::Close()
+{
+	if (CurrentState == EMediaState::Closed || CurrentState == EMediaState::Error)
+	{
+		return;
+	}
+
+	{
+		FScopeLock Lock(&CriticalSection);
+		bPlaybackCompleted = true;
+	}
+
+	InputWorker.DestroyThread();
+
+	CurrentState = EMediaState::Closed;
+
+	// remove delegates if registered
+	if (ResumeHandle.IsValid())
+	{
+		FCoreDelegates::ApplicationHasEnteredForegroundDelegate.Remove(ResumeHandle);
+		ResumeHandle.Reset();
+	}
+
+	if (PauseHandle.IsValid())
+	{
+		FCoreDelegates::ApplicationWillEnterBackgroundDelegate.Remove(PauseHandle);
+		PauseHandle.Reset();
+	}
+
+	if (MLHandleIsValid(VideoCodecHandle))
+	{
+		if (GSupportsImageExternal)
+		{
+			if (FLuminPlatformMisc::ShouldUseVulkan())
+			{
+				// Unregister the external texture on render thread
+				struct FReleaseVideoResourcesParams
+				{
+					FMagicLeapMediaCodecPlayer* MediaPlayer;
+					TSharedPtr<FMagicLeapVideoTextureDataVK, ESPMode::ThreadSafe> TextureData;
+					FGuid PlayerGuid;
+					MLHandle VideoCodecHandle;
+				};
+
+				FReleaseVideoResourcesParams ReleaseVideoResourcesParams = { this, StaticCastSharedPtr<FMagicLeapVideoTextureDataVK>(TextureData), PlayerGuid, VideoCodecHandle };
+
+				ENQUEUE_UNIQUE_RENDER_COMMAND_ONEPARAMETER(MagicLeapMediaPlayerDestroy, FReleaseVideoResourcesParams, Params, ReleaseVideoResourcesParams,
+				{
+					FExternalTextureRegistry::Get().UnregisterExternalTexture(Params.PlayerGuid);
+					Params.TextureData->bIsVideoTextureValid = false;
+
+					if (Params.TextureData->PreviousNativeBuffer != 0 && MLHandleIsValid(Params.TextureData->PreviousNativeBuffer))
+					{
+						Params.MediaPlayer->ReleaseNativeBuffer_RenderThread(Params.VideoCodecHandle, Params.TextureData->PreviousNativeBuffer);
+						Params.TextureData->PreviousNativeBuffer = 0;
+						Params.TextureData->VideoTexturePool.Empty();
+					}
+				});
+			}
+			else
+			{
+				// Unregister the external texture on render thread
+				struct FReleaseVideoResourcesParams
+				{
+					FMagicLeapMediaCodecPlayer* MediaPlayer;
+					// Can we make this a TWeakPtr? We need to ensure that FMagicLeapMediaCodecPlayer::TextureData is not destroyed before
+					// we unregister the external texture and releae the VideoTexture.
+					TSharedPtr<FMagicLeapVideoTextureDataGL, ESPMode::ThreadSafe> TextureData;
+					FGuid PlayerGuid;
+					MLHandle VideoCodecHandle;
+				};
+
+				FReleaseVideoResourcesParams ReleaseVideoResourcesParams = { this, StaticCastSharedPtr<FMagicLeapVideoTextureDataGL>(TextureData), PlayerGuid, VideoCodecHandle };
+
+				ENQUEUE_UNIQUE_RENDER_COMMAND_ONEPARAMETER(MagicLeapMediaPlayerDestroy, FReleaseVideoResourcesParams, Params, ReleaseVideoResourcesParams,
+				{
+					FExternalTextureRegistry::Get().UnregisterExternalTexture(Params.PlayerGuid);
+					Params.TextureData->bIsVideoTextureValid = false;
+
+					// @todo: this causes a crash
+					//Params.TextureData->VideoTexture->Release();
+					Params.TextureData->SaveContext();
+					Params.TextureData->MakeCurrent();
+
+					if (Params.TextureData->Image != EGL_NO_IMAGE_KHR)
+					{
+						eglDestroyImageKHR(eglGetCurrentDisplay(), Params.TextureData->Image);
+						Params.TextureData->Image = EGL_NO_IMAGE_KHR;
+					}
+
+					Params.TextureData->RestoreContext();
+					if (Params.TextureData->PreviousNativeBuffer != 0 && MLHandleIsValid(Params.TextureData->PreviousNativeBuffer))
+					{
+						Params.MediaPlayer->ReleaseNativeBuffer_RenderThread(Params.VideoCodecHandle, Params.TextureData->PreviousNativeBuffer);
+						Params.TextureData->PreviousNativeBuffer = 0;
+					}
+				});
+			}
+
+			FlushRenderingCommands();
+		}
+	}
+
+	StopAndReset(EMediaTrackType::Audio, AudioCodecHandle);
+	StopAndReset(EMediaTrackType::Video, VideoCodecHandle);
+
+	MLResult Result = MLMediaExtractorDestroy(MediaExtractorHandle);
+	UE_CLOG(Result != MLResult_Ok, LogMagicLeapMediaCodec, Error, TEXT("MLMediaExtractorDestroy failed with error %s"), UTF8_TO_TCHAR(MLGetResultString(Result)));
+	MediaExtractorHandle = ML_INVALID_HANDLE;
+
+	if (MLHandleIsValid(MediaDataSourceHandle))
+	{
+		Result = MLMediaDataSourceDestroy(MediaDataSourceHandle);
+		UE_CLOG(Result != MLResult_Ok, LogMagicLeapMediaCodec, Error, TEXT("MLMediaDataSourceDestroy failed with error %s"), UTF8_TO_TCHAR(MLGetResultString(Result)));
+		MediaDataSourceHandle = ML_INVALID_HANDLE;
+
+		DataSourceArchive.Reset();
+	}
+
+	switch (MediaSourceType)
+	{
+	case(EMediaSourceType::AudioOnly):
+		Result = MLMediaCodecDestroy(AudioCodecHandle);
+		UE_CLOG(Result != MLResult_Ok, LogMagicLeapMediaCodec, Error, TEXT("MLMediaCodecDestroy failed with error %s"), UTF8_TO_TCHAR(MLGetResultString(Result)));
+		break;
+	case(EMediaSourceType::VideoOnly):
+		Result = MLMediaCodecDestroy(VideoCodecHandle);
+		UE_CLOG(Result != MLResult_Ok, LogMagicLeapMediaCodec, Error, TEXT("MLMediaCodecDestroy failed with error %s"), UTF8_TO_TCHAR(MLGetResultString(Result)));
+		break;
+	case(EMediaSourceType::VideoAndAudio):
+		Result = MLMediaCodecDestroy(AudioCodecHandle);
+		UE_CLOG(Result != MLResult_Ok, LogMagicLeapMediaCodec, Error, TEXT("MLMediaCodecDestroy failed with error %s"), UTF8_TO_TCHAR(MLGetResultString(Result)));
+		Result = MLMediaCodecDestroy(VideoCodecHandle);
+		UE_CLOG(Result != MLResult_Ok, LogMagicLeapMediaCodec, Error, TEXT("MLMediaCodecDestroy failed with error %s"), UTF8_TO_TCHAR(MLGetResultString(Result)));
+		break;
+	}
+
+	// TODO: Destroy format handles in FTrackData
+
+	AudioCodecHandle = ML_INVALID_HANDLE;
+	VideoCodecHandle = ML_INVALID_HANDLE;
+
+	bMediaPrepared = false;
+	Info.Empty();
+	MediaUrl = FString();
+	AudioSamplePool->Reset();
+	TrackInfo.Empty();
+	SelectedTrack.Empty();
+
+	// notify listeners
+	EventSink.ReceiveMediaEvent(EMediaEvent::TracksChanged);
+	EventSink.ReceiveMediaEvent(EMediaEvent::MediaClosed);
+}
+
+IMediaCache& FMagicLeapMediaCodecPlayer::GetCache()
+{
+	return *this;
+}
+
+IMediaControls& FMagicLeapMediaCodecPlayer::GetControls()
+{
+	return *this;
+}
+
+FString FMagicLeapMediaCodecPlayer::GetInfo() const
+{
+	return Info;
+}
+
+FName FMagicLeapMediaCodecPlayer::GetPlayerName() const
+{
+	static FName PlayerName(TEXT("MagicLeapMediaCodec"));
+	return PlayerName;
+}
+
+IMediaSamples& FMagicLeapMediaCodecPlayer::GetSamples()
+{
+	return *Samples.Get();
+}
+
+FString FMagicLeapMediaCodecPlayer::GetStats() const
+{
+	return TEXT("MagicLeapMediaCodec stats information not implemented yet");
+}
+
+IMediaTracks& FMagicLeapMediaCodecPlayer::GetTracks()
+{
+	return *this;
+}
+
+FString FMagicLeapMediaCodecPlayer::GetUrl() const
+{
+	return MediaUrl;
+}
+
+IMediaView& FMagicLeapMediaCodecPlayer::GetView()
+{
+	return *this;
+}
+
+bool FMagicLeapMediaCodecPlayer::Open(const FString& Url, const IMediaOptions* Options)
+{
+	Close();
+
+	if (!MLHandleIsValid(MediaExtractorHandle))
+	{
+		MLResult Result = MLMediaExtractorCreate(&MediaExtractorHandle);
+		UE_CLOG(Result != MLResult_Ok, LogMagicLeapMediaCodec, Error, TEXT("MLMediaExtractorCreate failed with error %s."), UTF8_TO_TCHAR(MLGetResultString(Result)));
+		CurrentState = (Samples.IsValid() && Result == MLResult_Ok) ? EMediaState::Closed : EMediaState::Error;
+	}
+	if (CurrentState == EMediaState::Error)
+	{
+		EventSink.ReceiveMediaEvent(EMediaEvent::MediaOpenFailed);
+		return false;
+	}
+
+	if ((Url.IsEmpty()))
+	{
+		EventSink.ReceiveMediaEvent(EMediaEvent::MediaOpenFailed);
+		return false;
+	}
+
+	MediaUrl = Url;
+
+	const FString localFileSchema = "file://";
+
+	// open the media
+	if (Url.StartsWith(localFileSchema))
+	{
+		FString FilePath = Url.RightChop(localFileSchema.Len());
+		FPaths::NormalizeFilename(FilePath);
+
+		IPlatformFile& PlatformFile = IPlatformFile::GetPlatformPhysical();
+		// This module is only for Lumin so this is fine for now.
+		FLuminPlatformFile* LuminPlatformFile = static_cast<FLuminPlatformFile*>(&PlatformFile);
+		// make sure the file exists
+		if (!LuminPlatformFile->FileExists(*FilePath, FilePath))
+		{
+			UE_LOG(LogMagicLeapMediaCodec, Error, TEXT("File doesn't exist %s."), *FilePath);
+			EventSink.ReceiveMediaEvent(EMediaEvent::MediaOpenFailed);
+			return false;
+		}
+
+		const bool Precache = (Options != nullptr) ? Options->GetMediaOption("PrecacheFile", false) : false;
+
+		if (Precache)
+		{
+			UE_LOG(LogMagicLeapMediaCodec, Display, TEXT("Precaching media file %s"), *FilePath);
+
+			FArrayReader* Reader = new FArrayReader();
+
+			if (FFileHelper::LoadFileToArray(*Reader, *FilePath))
+			{
+				DataSourceArchive = MakeShareable(Reader);
+			}
+			else
+			{
+				delete Reader;
+			}
+
+			if (!DataSourceArchive.IsValid())
+			{
+				UE_LOG(LogMagicLeapMediaCodec, Error, TEXT("Failed to open or read media file %s"), *FilePath);
+				EventSink.ReceiveMediaEvent(EMediaEvent::MediaOpenFailed);
+				return false;
+			}
+
+			if (DataSourceArchive->TotalSize() == 0)
+			{
+				UE_LOG(LogMagicLeapMediaCodec, Error, TEXT("Cannot open media from empty file %s."), *FilePath);
+				EventSink.ReceiveMediaEvent(EMediaEvent::MediaOpenFailed);
+				return false;
+			}
+
+			MLResult Result = MLMediaDataSourceCreate(MediaDataSourceReadAt_Callback, MediaDataSourceGetSize_Callback, MediaDataSourceClose_Callback, reinterpret_cast<void*>(this), &MediaDataSourceHandle);
+			if (Result != MLResult_Ok)
+			{
+				UE_LOG(LogMagicLeapMediaCodec, Error, TEXT("MLMediaDataSourceCreate for path %s failed with error %s."), *FilePath, UTF8_TO_TCHAR(MLGetResultString(Result)));
+				EventSink.ReceiveMediaEvent(EMediaEvent::MediaOpenFailed);
+				return false;
+			}
+
+			Result = MLMediaExtractorSetMediaDataSource(MediaExtractorHandle, MediaDataSourceHandle);
+			if (Result != MLResult_Ok)
+			{
+				UE_LOG(LogMagicLeapMediaCodec, Error, TEXT("MLMediaExtractorSetMediaDataSource for path %s failed with error %s."), *FilePath, UTF8_TO_TCHAR(MLGetResultString(Result)));
+				EventSink.ReceiveMediaEvent(EMediaEvent::MediaOpenFailed);
+				return false;
+			}
+		}
+		else
+		{
+			MLResult Result = MLMediaExtractorSetDataSourceForPath(MediaExtractorHandle, TCHAR_TO_UTF8(*FilePath));
+			if (Result != MLResult_Ok)
+			{
+				UE_LOG(LogMagicLeapMediaCodec, Error, TEXT("MLMediaExtractorSetDataSourceForPath for path %s failed with error %s."), *FilePath, UTF8_TO_TCHAR(MLGetResultString(Result)));
+				EventSink.ReceiveMediaEvent(EMediaEvent::MediaOpenFailed);
+				return false;
+			}
+		}
+	}
+	else
+	{
+		// open remote media
+		MLResult Result = MLMediaExtractorSetDataSourceForURI(MediaExtractorHandle, TCHAR_TO_UTF8(*Url));
+		if (Result != MLResult_Ok)
+		{
+			UE_LOG(LogMagicLeapMediaCodec, Error, TEXT("MLMediaExtractorSetDataSourceForURI for remote media source %s failed with error %s."), *Url, UTF8_TO_TCHAR(MLGetResultString(Result)));
+			EventSink.ReceiveMediaEvent(EMediaEvent::MediaOpenFailed);
+			return false;
+		}
+	}
+
+	EventSink.ReceiveMediaEvent(EMediaEvent::MediaConnecting);
+
+	// prepare media
+	MediaUrl = Url;
+
+	if (!bMediaPrepared)
+	{
+		// TODO: For remote source, track information will probably not be available immediately.
+		GetTrackInformation();
+		bMediaPrepared = true;
+		bool bVideoCodecStatus = CreateMediaCodec(EMediaTrackType::Video, VideoCodecHandle);
+		bool bAudioCodecStatus = CreateMediaCodec(EMediaTrackType::Audio, AudioCodecHandle);
+		if (bVideoCodecStatus && !(bAudioCodecStatus))
+		{
+			MediaSourceType = EMediaSourceType::VideoOnly;
+		}
+		else if (bAudioCodecStatus && !(bVideoCodecStatus))
+		{
+			MediaSourceType = EMediaSourceType::AudioOnly;
+		}
+		else if (bVideoCodecStatus && bAudioCodecStatus)
+		{
+			MediaSourceType = EMediaSourceType::VideoAndAudio;
+		}
+		else
+		{
+			EventSink.ReceiveMediaEvent(EMediaEvent::MediaOpenFailed);
+			return false;
+		}
+	}
+
+	switch (MediaSourceType)
+	{
+	case(EMediaSourceType::VideoOnly):
+		StartMediaCodec(VideoCodecHandle);
+		break;
+	case(EMediaSourceType::AudioOnly):
+		StartMediaCodec(AudioCodecHandle);
+		break;
+	case(EMediaSourceType::VideoAndAudio):
+		StartMediaCodec(VideoCodecHandle);
+		StartMediaCodec(AudioCodecHandle);
+		break;
+	};
+
+	InputWorker.InitThread(*this, MediaExtractorHandle, CriticalSection, GT_IT_Mutex, RT_IT_Mutex);
+
+	CurrentState = EMediaState::Stopped;
+	// notify listeners
+	EventSink.ReceiveMediaEvent(EMediaEvent::TracksChanged);
+	EventSink.ReceiveMediaEvent(EMediaEvent::MediaOpened);
+
+	return true;
+}
+
+bool FMagicLeapMediaCodecPlayer::Open(const TSharedRef<FArchive, ESPMode::ThreadSafe>& Archive, const FString& OriginalUrl, const IMediaOptions* Options)
+{
+	// TODO: MagicLeapMedia: implement opening media from FArchive
+	return false;
+}
+
+bool FMagicLeapMediaCodecPlayer::CanControl(EMediaControl Control) const
+{
+	if (Control == EMediaControl::Pause)
+	{
+		return (CurrentState == EMediaState::Playing);
+	}
+
+	if (Control == EMediaControl::Resume)
+	{
+		return (CurrentState == EMediaState::Paused);
+	}
+
+	if (Control == EMediaControl::Seek)
+	{
+		return (CurrentState == EMediaState::Playing) || (CurrentState == EMediaState::Paused);
+	}
+
+	return false;
+}
+
+FTimespan FMagicLeapMediaCodecPlayer::GetDuration() const
+{
+	FTimespan Duration = FTimespan::Zero();
+
+	if (CurrentState == EMediaState::Playing || CurrentState == EMediaState::Paused || CurrentState == EMediaState::Stopped)
+	{
+		if (TrackInfo[EMediaTrackType::Video].Num() > 0)
+		{
+			return TrackInfo[EMediaTrackType::Video][SelectedTrack[EMediaTrackType::Video]].Duration;
+		}
+	}
+
+	return Duration;
+}
+
+float FMagicLeapMediaCodecPlayer::GetRate() const
+{
+	return (CurrentState == EMediaState::Playing) ? 1.0f : 0.0f;
+}
+
+EMediaState FMagicLeapMediaCodecPlayer::GetState() const
+{
+	return CurrentState;
+}
+
+EMediaStatus FMagicLeapMediaCodecPlayer::GetStatus() const
+{
+	return EMediaStatus::None;
+}
+
+TRangeSet<float> FMagicLeapMediaCodecPlayer::GetSupportedRates(EMediaRateThinning Thinning) const
+{
+	TRangeSet<float> Result;
+
+	Result.Add(TRange<float>(0.0f));
+	Result.Add(TRange<float>(1.0f));
+
+	return Result;
+}
+
+FTimespan FMagicLeapMediaCodecPlayer::GetTime() const
+{
+	if (CurrentState == EMediaState::Playing || CurrentState == EMediaState::Paused)
+	{
+		return CurrentPlaybackTime.Load();
+	}
+
+	return FTimespan::Zero();
+}
+
+bool FMagicLeapMediaCodecPlayer::IsLooping() const
+{
+	return bLoopPlayback;
+}
+
+bool FMagicLeapMediaCodecPlayer::Seek(const FTimespan& Time)
+{
+	if ((CurrentState == EMediaState::Closed) || (CurrentState == EMediaState::Error) || (CurrentState == EMediaState::Preparing))
+	{
+		UE_LOG(LogMagicLeapMediaCodec, Warning, TEXT("Cannot seek while closed, preparing, or in error state"));
+		return false;
+	}
+	else if (CurrentState == EMediaState::Playing || CurrentState == EMediaState::Paused)
+	{
+		InputWorker.Seek(Time);
+	}
+	return true;
+}
+
+bool FMagicLeapMediaCodecPlayer::SetLooping(bool Looping)
+{
+	bLoopPlayback = Looping;
+	return bLoopPlayback;
+}
+
+bool FMagicLeapMediaCodecPlayer::SetRate(float Rate)
+{
+	if ((CurrentState == EMediaState::Closed) || (CurrentState == EMediaState::Error) || (CurrentState == EMediaState::Preparing))
+	{
+		UE_LOG(LogMagicLeapMediaCodec, Warning, TEXT("Cannot set rate while closed, preparing, or in error state"));
+		return false;
+	}
+
+	if (Rate == GetRate())
+	{
+		// rate already set
+		return true;
+	}
+	{
+		// Scope lock for CurrentState and StartPresentationTime. These are read and written to in the render thread for the video codec.
+		FScopeLock Lock(&CriticalSection);
+		if (Rate == 0.0f)
+		{
+			CurrentState = EMediaState::Paused;
+			bResetVideoCodecStartTime = true;
+			EventSink.ReceiveMediaEvent(EMediaEvent::PlaybackSuspended);
+			return true;
+		}
+		else if (Rate == 1.0f)
+		{
+			if (CurrentState != EMediaState::Playing)
+			{
+				CurrentState = EMediaState::Playing;
+				EventSink.ReceiveMediaEvent(EMediaEvent::PlaybackResumed);
+				InputWorker.WakeUp();
+			}
+			return true;
+		}
+		else
+		{
+			UE_LOG(LogMagicLeapMediaCodec, Error, TEXT("Rate %f not supported by MagicLeapMedia."), Rate);
+			return false;
+		}
+	}
+
+	return false;
+}
+
+bool FMagicLeapMediaCodecPlayer::SetNativeVolume(float Volume)
+{
+	UE_LOG(LogMagicLeapMediaCodec, Warning, TEXT("SetNativeVolume() is not supported for MagicLeapMedia. Use UMediaSoundComponent::SetVolumeMultiplier() instead."));
+	return false;
+}
+
+void FMagicLeapMediaCodecPlayer::SetGuid(const FGuid& Guid)
+{
+	PlayerGuid = Guid;
+}
+
+void FMagicLeapMediaCodecPlayer::TickFetch(FTimespan DeltaTime, FTimespan Timecode)
+{
+	// TODO: this condition needs to be revised based on how we handle EMediaState::Preparing as preapring would probably require processing some output samples.
+	// We should not process output samples whem playback is paused. Not even to flush the already decoded samples as we need a frame accurate sync.
+	if (CurrentState != EMediaState::Playing)
+	{
+		return;
+	}
+
+	if (GSupportsImageExternal)
+	{
+		if (FLuminPlatformMisc::ShouldUseVulkan())
+		{
+			struct FWriteVideoSampleParams
+			{
+				FMagicLeapMediaCodecPlayer* MediaPlayer;
+				TWeakPtr<FMagicLeapVideoTextureDataVK, ESPMode::ThreadSafe> TextureData;
+				FGuid PlayerGuid;
+				MLHandle VideoCodecHandle;
+				FTimespan FrameTimecode;
+			};
+
+			FWriteVideoSampleParams WriteVideoSampleParams = { this, StaticCastSharedPtr<FMagicLeapVideoTextureDataVK>(TextureData), PlayerGuid, VideoCodecHandle, Timecode };
+
+			ENQUEUE_UNIQUE_RENDER_COMMAND_ONEPARAMETER(MagicLeapMediaPlayerWriteVideoSample, FWriteVideoSampleParams, Params, WriteVideoSampleParams,
+			{
+				auto TextureDataPtr = Params.TextureData.Pin();
+
+				Params.MediaPlayer->bIsBufferAvailable = Params.MediaPlayer->ProcessVideoOutputSample_RenderThread(Params.VideoCodecHandle, Params.MediaPlayer->TrackInfo[EMediaTrackType::Video][Params.MediaPlayer->SelectedTrack[EMediaTrackType::Video]], Params.FrameTimecode);
+				if (!Params.MediaPlayer->IsBufferAvailable_RenderThread(Params.VideoCodecHandle))
+				{
+					// UE_LOG(LogMagicLeapMediaCodec, Error, TEXT("No video buffer available at TickFetch!"));
+					return;
+				}
+
+				if (TextureDataPtr->PreviousNativeBuffer != 0 && MLHandleIsValid(TextureDataPtr->PreviousNativeBuffer))
+				{
+					Params.MediaPlayer->ReleaseNativeBuffer_RenderThread(Params.VideoCodecHandle, TextureDataPtr->PreviousNativeBuffer);
+					TextureDataPtr->PreviousNativeBuffer = 0;
+				}
+
+				MLHandle NativeBuffer = ML_INVALID_HANDLE;
+				if (!Params.MediaPlayer->GetNativeBuffer_RenderThread(Params.VideoCodecHandle, NativeBuffer))
+				{
+					return;
+				}
+
+				check(MLHandleIsValid(NativeBuffer));
+
+				if (Params.MediaPlayer->bPlaybackCompleted_RenderThread)
+				{
+					TextureDataPtr->VideoTexturePool.Empty();
+				}
+
+				if (!TextureDataPtr->VideoTexturePool.Contains((uint64)NativeBuffer))
+				{
+					FTextureRHIRef NewMediaTexture;
+					if (!FMagicLeapHelperVulkan::GetMediaTexture(NewMediaTexture, TextureDataPtr->VideoSampler, NativeBuffer))
+					{
+						UE_LOG(LogMagicLeapMediaCodec, Error, TEXT("Failed to get next media texture."));
+						return;
+					}
+
+					TextureDataPtr->VideoTexturePool.Add((uint64)NativeBuffer, NewMediaTexture);
+					
+					if (TextureDataPtr->VideoTexture == nullptr)
+					{
+						FRHIResourceCreateInfo CreateInfo;
+						TextureDataPtr->VideoTexture = RHICmdList.CreateTextureExternal2D(1, 1, PF_R8G8B8A8, 1, 1, 0, CreateInfo);
+					}
+
+					FMagicLeapHelperVulkan::AliasMediaTexture(TextureDataPtr->VideoTexture, NewMediaTexture);
+				}
+				else
+				{
+					FTextureRHIRef* const PooledMediaTexture = TextureDataPtr->VideoTexturePool.Find((uint64)NativeBuffer);
+					check(PooledMediaTexture != nullptr);
+					FMagicLeapHelperVulkan::AliasMediaTexture(TextureDataPtr->VideoTexture, *PooledMediaTexture);
+				}
+
+				TextureDataPtr->bIsVideoTextureValid = (TextureDataPtr->bIsVideoTextureValid && !Params.MediaPlayer->UpdateTransformMatrix_RenderThread(Params.VideoCodecHandle));
+
+				if (!TextureDataPtr->bIsVideoTextureValid)
+				{
+					Params.MediaPlayer->RegisterExternalTexture_RenderThread(Params.PlayerGuid, TextureDataPtr->VideoTexture, TextureDataPtr->VideoSampler);
+					TextureDataPtr->bIsVideoTextureValid = true;
+				}
+
+				TextureDataPtr->PreviousNativeBuffer = NativeBuffer;
+			});
+		}
+		else
+		{
+			struct FWriteVideoSampleParams
+			{
+				FMagicLeapMediaCodecPlayer* MediaPlayer;
+				TWeakPtr<FMagicLeapVideoTextureDataGL, ESPMode::ThreadSafe> TextureData;
+				FGuid PlayerGuid;
+				MLHandle VideoCodecHandle;
+				FTimespan FrameTimecode;
+			};
+
+			FWriteVideoSampleParams WriteVideoSampleParams = { this, StaticCastSharedPtr<FMagicLeapVideoTextureDataGL>(TextureData), PlayerGuid, VideoCodecHandle, Timecode };
+
+			ENQUEUE_UNIQUE_RENDER_COMMAND_ONEPARAMETER(MagicLeapMediaPlayerWriteVideoSample, FWriteVideoSampleParams, Params, WriteVideoSampleParams,
+			{
+				auto TextureDataPtr = Params.TextureData.Pin();
+
+				Params.MediaPlayer->bIsBufferAvailable = Params.MediaPlayer->ProcessVideoOutputSample_RenderThread(Params.VideoCodecHandle, Params.MediaPlayer->TrackInfo[EMediaTrackType::Video][Params.MediaPlayer->SelectedTrack[EMediaTrackType::Video]], Params.FrameTimecode);
+				if (!Params.MediaPlayer->IsBufferAvailable_RenderThread(Params.VideoCodecHandle))
+				{
+					return;
+				}
+
+				FTextureRHIRef MediaVideoTexture = TextureDataPtr->VideoTexture;
+				if (MediaVideoTexture == nullptr)
+				{
+					FRHIResourceCreateInfo CreateInfo;
+					MediaVideoTexture = RHICmdList.CreateTextureExternal2D(1, 1, PF_R8G8B8A8, 1, 1, 0, CreateInfo);
+					TextureDataPtr->VideoTexture = MediaVideoTexture;
+
+					if (MediaVideoTexture == nullptr)
+					{
+						UE_LOG(LogMagicLeapMediaCodec, Warning, TEXT("CreateTextureExternal2D failed!"));
+						return;
+					}
+
+					TextureDataPtr->bIsVideoTextureValid = false;
+				}
+
+				// MLHandle because Unreal's uint64 is 'unsigned long long *' whereas uint64_t for the C-API is 'unsigned long *'
+				// TODO: Fix the Unreal types for the above comment.
+				MLHandle nativeBuffer = ML_INVALID_HANDLE;
+				if (!Params.MediaPlayer->GetNativeBuffer_RenderThread(Params.VideoCodecHandle, nativeBuffer))
+				{
+					return;
+				}
+
+				int32 CurrentFramePosition = 0;
+				if (!Params.MediaPlayer->GetCurrentPosition_RenderThread(Params.VideoCodecHandle, CurrentFramePosition))
+				{
+					return;
+				}
+
+				// Clear gl errors as they can creep in from the UE4 renderer.
+				glGetError();
+
+				if (!TextureDataPtr->bContextCreated)
+				{
+					TextureDataPtr->InitContext();
+					TextureDataPtr->bContextCreated = true;
+				}
+				TextureDataPtr->SaveContext();
+				TextureDataPtr->MakeCurrent();
+
+				int32 TextureID = *reinterpret_cast<int32*>(MediaVideoTexture->GetNativeResource());
+				if (TextureDataPtr->Image != EGL_NO_IMAGE_KHR)
+				{
+					eglDestroyImageKHR(eglGetCurrentDisplay(), TextureDataPtr->Image);
+					TextureDataPtr->Image = EGL_NO_IMAGE_KHR;
+				}
+				if (TextureDataPtr->PreviousNativeBuffer != 0 && MLHandleIsValid(TextureDataPtr->PreviousNativeBuffer))
+				{
+					Params.MediaPlayer->ReleaseNativeBuffer_RenderThread(Params.VideoCodecHandle, TextureDataPtr->PreviousNativeBuffer);
+				}
+				TextureDataPtr->PreviousNativeBuffer = nativeBuffer;
+
+				// Wrap latest decoded frame into a new gl texture oject
+				TextureDataPtr->Image = eglCreateImageKHR(TextureDataPtr->Display, EGL_NO_CONTEXT, EGL_NATIVE_BUFFER_ANDROID, (EGLClientBuffer)(void *)nativeBuffer, NULL);
+				if (TextureDataPtr->Image == EGL_NO_IMAGE_KHR)
+				{
+					EGLint errorcode = eglGetError();
+					UE_LOG(LogMagicLeapMediaCodec, Error, TEXT("Failed to create EGLImage from the buffer. %d"), errorcode);
+					TextureDataPtr->RestoreContext();
+					return;
+				}
+				glActiveTexture(GL_TEXTURE0);
+				glBindTexture(GL_TEXTURE_EXTERNAL_OES, TextureID);
+				glEGLImageTargetTexture2DOES(GL_TEXTURE_EXTERNAL_OES, TextureDataPtr->Image);
+				glTexParameteri(GL_TEXTURE_EXTERNAL_OES, GL_TEXTURE_MIN_FILTER, GL_LINEAR);
+				glTexParameteri(GL_TEXTURE_EXTERNAL_OES, GL_TEXTURE_MAG_FILTER, GL_LINEAR);
+				glTexParameteri(GL_TEXTURE_EXTERNAL_OES, GL_TEXTURE_WRAP_S, GL_CLAMP_TO_EDGE);
+				glTexParameteri(GL_TEXTURE_EXTERNAL_OES, GL_TEXTURE_WRAP_T, GL_CLAMP_TO_EDGE);
+				glBindTexture(GL_TEXTURE_EXTERNAL_OES, 0);
+
+				TextureDataPtr->RestoreContext();
+
+				TextureDataPtr->bIsVideoTextureValid = Params.MediaPlayer->UpdateTransformMatrix_RenderThread(Params.VideoCodecHandle);
+
+				if (!TextureDataPtr->bIsVideoTextureValid)
+				{
+					FSamplerStateInitializerRHI SamplerStateInitializer(SF_Bilinear, AM_Clamp, AM_Clamp, AM_Clamp);
+					FSamplerStateRHIRef SamplerStateRHI = RHICreateSamplerState(SamplerStateInitializer);
+					Params.MediaPlayer->RegisterExternalTexture_RenderThread(Params.PlayerGuid, MediaVideoTexture, SamplerStateRHI);
+					TextureDataPtr->bIsVideoTextureValid = true;
+				}
+			});
+		}
+	}
+}
+
+bool FMagicLeapMediaCodecPlayer::UpdateTransformMatrix_RenderThread(MLHandle InVideoCodecHandle)
+{
+	MLResult Result = MLMediaCodecGetFrameTransformationMatrix(InVideoCodecHandle, FrameTransformationMatrix);
+	if (UScale != FrameTransformationMatrix[0] ||
+		UOffset != FrameTransformationMatrix[12] ||
+		VScale != -FrameTransformationMatrix[5] ||
+		VOffset != (1.0f - FrameTransformationMatrix[13]))
+	{
+		UScale = FrameTransformationMatrix[0];
+		UOffset = FrameTransformationMatrix[12];
+		VScale = -FrameTransformationMatrix[5];
+		VOffset = 1.0f - FrameTransformationMatrix[13];
+		return true;
+	}
+
+	return false;
+}
+
+void FMagicLeapMediaCodecPlayer::RegisterExternalTexture_RenderThread(const FGuid& InGuid, FTextureRHIRef& InTextureRHI, FSamplerStateRHIRef& InSamplerStateRHI)
+{
+	FExternalTextureRegistry::Get().RegisterExternalTexture(InGuid, InTextureRHI, InSamplerStateRHI, FLinearColor(UScale, 0.0f, 0.0f, VScale), FLinearColor(UOffset, VOffset, 0.0f, 0.0f));
+}
+
+void FMagicLeapMediaCodecPlayer::TickInput(FTimespan DeltaTime, FTimespan Timecode)
+{
+	if (!bMediaPrepared)
+	{
+		return;
+	}
+
+	// Fire all pending media events. Likely queued from MediaCodecInputWorker thread.
+	EMediaEvent PendingMediaEvent;
+	while (MediaEventQueue.Dequeue(PendingMediaEvent))
+	{
+		EventSink.ReceiveMediaEvent(PendingMediaEvent);
+	}
+
+	{
+		FScopeLock Lock(&CriticalSection);
+		// TODO: bPlaybackCompleted is updated on the render thread as well as the input worker thread.
+		// So consider firing this event based on an event queue instead of directly testing for the condition here.
+		if (bPlaybackCompleted)
+		{
+			if (!IsLooping())
+			{
+				CurrentState = EMediaState::Stopped;
+			}
+			EventSink.ReceiveMediaEvent(EMediaEvent::PlaybackEndReached);
+		}
+	}
+
+	if (CurrentState != EMediaState::Playing)
+	{
+		// remove delegates if registered
+		if (ResumeHandle.IsValid())
+		{
+			FCoreDelegates::ApplicationHasEnteredForegroundDelegate.Remove(ResumeHandle);
+			ResumeHandle.Reset();
+		}
+
+		if (PauseHandle.IsValid())
+		{
+			FCoreDelegates::ApplicationWillEnterBackgroundDelegate.Remove(PauseHandle);
+			PauseHandle.Reset();
+		}
+	}
+
+	// register delegate if not registered
+	if (!ResumeHandle.IsValid())
+	{
+		ResumeHandle = FCoreDelegates::ApplicationHasEnteredForegroundDelegate.AddRaw(this, &FMagicLeapMediaCodecPlayer::HandleApplicationHasEnteredForeground);
+	}
+	if (!PauseHandle.IsValid())
+	{
+		PauseHandle = FCoreDelegates::ApplicationWillEnterBackgroundDelegate.AddRaw(this, &FMagicLeapMediaCodecPlayer::HandleApplicationWillEnterBackground);
+	}
+}
+
+void FMagicLeapMediaCodecPlayer::TickAudio(FTimespan LastAudioSampleTime)
+{
+	if (GetSelectedTrack(EMediaTrackType::Audio) != INDEX_NONE)
+	{
+		ProcessAudioOutputSample(AudioCodecHandle, TrackInfo[EMediaTrackType::Audio][SelectedTrack[EMediaTrackType::Audio]], LastAudioSampleTime);
+	}
+}
+
+bool FMagicLeapMediaCodecPlayer::GetAudioTrackFormat(int32 TrackIndex, int32 FormatIndex, FMediaAudioTrackFormat& OutFormat) const
+{
+	if ((FormatIndex != 0) || TrackIndex >= TrackInfo[EMediaTrackType::Audio].Num() || TrackIndex < 0)
+	{
+		return false;
+	}
+
+	OutFormat.BitsPerSample = 16;
+	OutFormat.NumChannels = TrackInfo[EMediaTrackType::Audio][TrackIndex].ChannelCount;
+	OutFormat.SampleRate = TrackInfo[EMediaTrackType::Audio][TrackIndex].SampleRate;
+	OutFormat.TypeName = TrackInfo[EMediaTrackType::Audio][TrackIndex].FormatName;
+
+	return true;
+}
+
+int32 FMagicLeapMediaCodecPlayer::GetNumTracks(EMediaTrackType TrackType) const
+{
+	if (TrackInfo.Contains(TrackType))
+	{
+		return TrackInfo[TrackType].Num();
+	}
+
+	return 0;
+}
+
+int32 FMagicLeapMediaCodecPlayer::GetNumTrackFormats(EMediaTrackType TrackType, int32 TrackIndex) const
+{
+	return ((TrackIndex >= 0) && (TrackIndex < GetNumTracks(TrackType))) ? 1 : 0;
+}
+
+int32 FMagicLeapMediaCodecPlayer::GetSelectedTrack(EMediaTrackType TrackType) const
+{
+	if (SelectedTrack.Contains(TrackType))
+	{
+		return SelectedTrack[TrackType];
+	}
+
+	return INDEX_NONE;
+}
+
+FText FMagicLeapMediaCodecPlayer::GetTrackDisplayName(EMediaTrackType TrackType, int32 TrackIndex) const
+{
+	return FText::GetEmpty();
+}
+
+int32 FMagicLeapMediaCodecPlayer::GetTrackFormat(EMediaTrackType TrackType, int32 TrackIndex) const
+{
+	return (GetSelectedTrack(TrackType) != INDEX_NONE) ? 0 : INDEX_NONE;
+}
+
+FString FMagicLeapMediaCodecPlayer::GetTrackLanguage(EMediaTrackType TrackType, int32 TrackIndex) const
+{
+	if (TrackIndex >= 0 && TrackIndex < TrackInfo[TrackType].Num())
+	{
+		return TrackInfo[TrackType][TrackIndex].Language;
+	}
+	return FString();
+}
+
+FString FMagicLeapMediaCodecPlayer::GetTrackName(EMediaTrackType TrackType, int32 TrackIndex) const
+{
+	// Track names not supported in ML.
+	return FString();
+}
+
+bool FMagicLeapMediaCodecPlayer::GetVideoTrackFormat(int32 TrackIndex, int32 FormatIndex, FMediaVideoTrackFormat& OutFormat) const
+{
+	if ((FormatIndex != 0) || TrackIndex >= TrackInfo[EMediaTrackType::Video].Num())
+	{
+		return false;
+	}
+
+	OutFormat.Dim = FIntPoint(TrackInfo[EMediaTrackType::Video][TrackIndex].Width, TrackInfo[EMediaTrackType::Video][TrackIndex].Height);
+	OutFormat.FrameRate = static_cast<float>(TrackInfo[EMediaTrackType::Video][TrackIndex].FrameRate);
+	OutFormat.FrameRates = TRange<float>(OutFormat.FrameRate);
+	OutFormat.TypeName = TrackInfo[EMediaTrackType::Video][TrackIndex].FormatName;
+
+	return true;
+}
+
+bool FMagicLeapMediaCodecPlayer::SelectTrack(EMediaTrackType TrackType, int32 TrackIndex)
+{
+	if (TrackInfo.Contains(TrackType) && CurrentState != EMediaState::Preparing)
+	{
+		if (TrackInfo[TrackType].IsValidIndex(TrackIndex))
+		{
+			// TODO: codec needs to be changed.
+			MLResult Result = MLMediaExtractorSelectTrack(MediaExtractorHandle, static_cast<SIZE_T>(TrackInfo[TrackType][TrackIndex].TrackIndex));
+			if (Result == MLResult_Ok)
+			{
+				SelectedTrack[TrackType] = TrackIndex;
+				return true;
+			}
+			else
+			{
+				UE_LOG(LogMagicLeapMediaCodec, Error, TEXT("MLMediaExtractorSelectTrack failed with error %s"), UTF8_TO_TCHAR(MLGetResultString(Result)));
+			}
+		}
+	}
+	return false;
+}
+
+bool FMagicLeapMediaCodecPlayer::SetTrackFormat(EMediaTrackType TrackType, int32 TrackIndex, int32 FormatIndex)
+{
+	return false;
+}
+
+void FMagicLeapMediaCodecPlayer::QueueMediaEvent(EMediaEvent MediaEvent)
+{
+	MediaEventQueue.Enqueue(MediaEvent);
+}
+
+void FMagicLeapMediaCodecPlayer::FlushCodecs()
+{
+	MLResult Result;
+	if (MLHandleIsValid(AudioCodecHandle))
+	{
+		Result = MLMediaCodecFlush(AudioCodecHandle);
+		UE_CLOG(Result != MLResult_Ok, LogMagicLeapMediaCodec, Error, TEXT("MLMediaCodecFlush() failed with error %s"), UTF8_TO_TCHAR(MLMediaResultGetString(Result)));
+	}
+	if (MLHandleIsValid(VideoCodecHandle))
+	{
+		Result = MLMediaCodecFlush(VideoCodecHandle);
+		UE_CLOG(Result != MLResult_Ok, LogMagicLeapMediaCodec, Error, TEXT("MLMediaCodecFlush() failed with error %s"), UTF8_TO_TCHAR(MLMediaResultGetString(Result)));
+		// UE_LOG(LogMagicLeapMediaCodec, Warning, TEXT("video codec flushed."));
+	}
+
+	for (const auto& TrackPair : SelectedTrack)
+	{
+		if (TrackPair.Value != INDEX_NONE)
+		{
+			TrackInfo[TrackPair.Key][TrackPair.Value].bBufferPendingRender = false;
+			if (TrackPair.Key == EMediaTrackType::Audio)
+			{
+				TrackInfo[TrackPair.Key][TrackPair.Value].LastPresentationTime = FTimespan::Zero();
+			}
+		}
+	}
+
+	Samples->FlushSamples();
+}
+
+bool FMagicLeapMediaCodecPlayer::GetCodecForTrackIndex(int64_t TrackIndex, MLHandle& SampleCodecHandle) const
+{
+	if (TrackIndex == TrackInfo[EMediaTrackType::Video][SelectedTrack[EMediaTrackType::Video]].TrackIndex)
+	{
+		SampleCodecHandle = VideoCodecHandle;
+	}
+	else if (TrackIndex == TrackInfo[EMediaTrackType::Audio][SelectedTrack[EMediaTrackType::Audio]].TrackIndex)
+	{
+		SampleCodecHandle = AudioCodecHandle;
+	}
+	else
+	{
+		return false;
+	}
+
+	return true;
+}
+
+bool FMagicLeapMediaCodecPlayer::IsPlaybackCompleted() const
+{
+	return bPlaybackCompleted;
+}
+
+void FMagicLeapMediaCodecPlayer::SetPlaybackCompleted(bool PlaybackCompleted)
+{
+	bPlaybackCompleted = PlaybackCompleted;
+}
+
+void FMagicLeapMediaCodecPlayer::QueueVideoCodecStartTimeReset()
+{
+	bResetVideoCodecStartTime = true;
+}
+
+int64_t FMagicLeapMediaCodecPlayer::MediaDataSourceReadAt(MLHandle media_data_source, size_t position, size_t size, uint8_t *buffer)
+{
+	if (media_data_source == MediaDataSourceHandle && DataSourceArchive.IsValid() && buffer != nullptr)
+	{
+		int64_t BytesToRead = size;
+
+		DataSourceArchive->Seek(position);
+		int64 Size = DataSourceArchive->TotalSize();
+
+		if (BytesToRead > Size)
+		{
+			BytesToRead = Size;
+		}
+
+		if ((Size - BytesToRead) < DataSourceArchive->Tell())
+		{
+			BytesToRead = Size - position;
+		}
+
+		if (BytesToRead > 0)
+		{
+			DataSourceArchive->Serialize(buffer, BytesToRead);
+		}
+
+		return BytesToRead;
+	}
+
+	return -1;
+}
+
+int64_t FMagicLeapMediaCodecPlayer::MediaDataSourceGetSize(MLHandle media_data_source)
+{
+	if (media_data_source == MediaDataSourceHandle && DataSourceArchive.IsValid())
+	{
+		return DataSourceArchive->TotalSize();
+	}
+
+	return -1;
+}
+
+void FMagicLeapMediaCodecPlayer::MediaDataSourceClose(MLHandle media_data_source)
+{
+
+}
+
+void FMagicLeapMediaCodecPlayer::HandleApplicationHasEnteredForeground()
+{
+	// check state in case changed before ticked
+	if (CurrentState == EMediaState::Paused && bWasMediaPlayingBeforeAppPause)
+	{
+		// pause
+		SetRate(1.0f);
+	}
+}
+
+void FMagicLeapMediaCodecPlayer::HandleApplicationWillEnterBackground()
+{
+	bWasMediaPlayingBeforeAppPause = (CurrentState == EMediaState::Playing);
+	// check state in case changed before ticked
+	if (bWasMediaPlayingBeforeAppPause)
+	{
+		// pause
+		SetRate(0.0f);
+	}
+}
+
+FIntPoint FMagicLeapMediaCodecPlayer::GetVideoDimensions() const
+{
+	int32 width = 0;
+	int32 height = 0;
+	FMediaVideoTrackFormat VideoTrackFormat;
+	if (GetVideoTrackFormat(SelectedTrack[EMediaTrackType::Video], 0, VideoTrackFormat))
+	{
+		return VideoTrackFormat.Dim;
+	}
+	return FIntPoint(width, height);
+}
+
+bool FMagicLeapMediaCodecPlayer::StopAndReset(EMediaTrackType TrackType, MLHandle& CodecHandle)
+{
+	if (CodecHandle == ML_INVALID_HANDLE)
+	{
+		return true;
+	}
+	MLResult FlushResult = MLMediaCodecFlush(CodecHandle);
+	MLResult CodecStopResult = MLMediaCodecStop(CodecHandle);
+	if ((FlushResult != MLResult_Ok) || (CodecStopResult != MLResult_Ok))
+	{
+		UE_CLOG(FlushResult != MLResult_Ok, LogMagicLeapMediaCodec, Error, TEXT("MLMediaCodecFlush failed with error %d"), FlushResult);
+		UE_CLOG(CodecStopResult != MLResult_Ok, LogMagicLeapMediaCodec, Error, TEXT("MLMediaCodecStop failed with error %d"), CodecStopResult);
+		return false;
+	}
+	return true;
+}
+
+bool FMagicLeapMediaCodecPlayer::SetRateOne()
+{
+	StartMediaCodec(VideoCodecHandle);
+	StartMediaCodec(AudioCodecHandle);
+
+	CurrentState = EMediaState::Playing;
+	EventSink.ReceiveMediaEvent(EMediaEvent::PlaybackResumed);
+	return true;
+}
+
+void FMagicLeapMediaCodecPlayer::GetTrackInformation()
+{
+	uint64 NumTracks = 0;
+	MLResult Result = MLMediaExtractorGetTrackCount(MediaExtractorHandle, reinterpret_cast<uint64_t*>(&NumTracks));
+	UE_CLOG(Result != MLResult_Ok, LogMagicLeapMediaCodec, Error, TEXT("MLMediaExtractorGetTrackCount() failed with error %s"), UTF8_TO_TCHAR(MLMediaResultGetString(Result)));
+
+	char Mime[MAX_KEY_STRING_SIZE] = "";
+	FString MimeTypeString;
+
+	TrackInfo.Add(EMediaTrackType::Video);
+	TrackInfo.Add(EMediaTrackType::Audio);
+	TrackInfo.Add(EMediaTrackType::Caption);
+	TrackInfo.Add(EMediaTrackType::Subtitle);
+	TrackInfo.Add(EMediaTrackType::Metadata);
+
+	SelectedTrack.Add(EMediaTrackType::Video, INDEX_NONE);
+	SelectedTrack.Add(EMediaTrackType::Audio, INDEX_NONE);
+	SelectedTrack.Add(EMediaTrackType::Caption, INDEX_NONE);
+	SelectedTrack.Add(EMediaTrackType::Subtitle, INDEX_NONE);
+	SelectedTrack.Add(EMediaTrackType::Metadata, INDEX_NONE);
+
+	for (SIZE_T TrackIndex = 0; TrackIndex < NumTracks; ++TrackIndex)
+	{
+		MLHandle TrackFormatHandle = ML_INVALID_HANDLE;
+		Result = MLMediaExtractorGetTrackFormat(MediaExtractorHandle, TrackIndex, &TrackFormatHandle);
+		if (Result != MLResult_Ok)
+		{
+			UE_LOG(LogMagicLeapMediaCodec, Error, TEXT("MLMediaExtractorGetTrackFormat() failed with error %s"), UTF8_TO_TCHAR(MLMediaResultGetString(Result)));
+			continue;
+		}
+
+		// HACK: Fix this!
+		// Keys should be defined in libs, ran into linking issue, SDK-1402
+		Result = MLMediaFormatGetKeyString(TrackFormatHandle, MLMediaFormat_Key_Mime, Mime);
+		UE_CLOG(Result != MLResult_Ok, LogMagicLeapMediaCodec, Error, TEXT("MLMediaFormatGetKeyString(mime) failed with error %s"), UTF8_TO_TCHAR(MLMediaResultGetString(Result)));
+
+		MimeTypeString = FString(Mime);
+
+		// TODO: refactor this. Lot of duplication!
+		if (MimeTypeString.Contains("audio/"))
+		{
+			int32 Index = TrackInfo[EMediaTrackType::Audio].Add(FTrackData(MimeTypeString, TrackIndex));
+			UpdateAudioTrackInfo(TrackFormatHandle, TrackInfo[EMediaTrackType::Audio][Index]);
+
+			SelectedTrack[EMediaTrackType::Audio] = 0;
+		}
+		else if (MimeTypeString.Contains("video/"))
+		{
+			int32 Index = TrackInfo[EMediaTrackType::Video].Add(FTrackData(MimeTypeString, TrackIndex));
+			UpdateVideoTrackInfo(TrackFormatHandle, TrackInfo[EMediaTrackType::Video][Index]);
+
+			SelectedTrack[EMediaTrackType::Video] = 0;
+		}
+		// What are the mimetypes for captions/subtitle/metadata?
+		// TODO: Process subtitles and metadata
+		else if (MimeTypeString.Contains("text/"))
+		{
+		}
+
+		Result = MLMediaFormatDestroy(TrackFormatHandle);
+		UE_CLOG(Result != MLResult_Ok, LogMagicLeapMediaCodec, Error, TEXT("MLMediaFormatDestroy() failed with error %s"), UTF8_TO_TCHAR(MLMediaResultGetString(Result)));
+	}
+}
+
+void FMagicLeapMediaCodecPlayer::UpdateCommonTrackInfo(MLHandle FormatHandle, FTrackData& CurrentTrackData)
+{
+	char FormatName[MAX_FORMAT_STRING_SIZE] = "";
+	MLResult Result = MLMediaFormatObjectToString(FormatHandle, FormatName);
+	UE_CLOG(Result != MLResult_Ok, LogMagicLeapMediaCodec, Error, TEXT("MLMediaFormatObjectToString() failed with error %s"), UTF8_TO_TCHAR(MLMediaResultGetString(Result)));
+	CurrentTrackData.FormatName = (Result == MLResult_Ok) ? FString(FormatName) : FString();
+
+	int64 Duration = 0;
+	Result = MLMediaFormatGetKeyValueInt64(FormatHandle, MLMediaFormat_Key_Duration, reinterpret_cast<int64_t*>(&Duration));
+	UE_CLOG(Result != MLResult_Ok, LogMagicLeapMediaCodec, Error, TEXT("MLMediaFormatGetKeyValueInt64(duration) failed with error %s"), UTF8_TO_TCHAR(MLMediaResultGetString(Result)));
+	CurrentTrackData.Duration = (Result == MLResult_Ok) ? FTimespan::FromMicroseconds(Duration) : CurrentTrackData.Duration;
+
+	char Language[MAX_KEY_STRING_SIZE] = "";
+	Result = MLMediaFormatGetKeyString(FormatHandle, MLMediaFormat_Key_Language, Language);
+	UE_CLOG(Result != MLResult_Ok, LogMagicLeapMediaCodec, Error, TEXT("MLMediaFormatGetKeyString(language) failed with error %s"), UTF8_TO_TCHAR(MLMediaResultGetString(Result)));
+	CurrentTrackData.Language = (Result == MLResult_Ok) ? FString(Language) : FString();
+}
+
+void FMagicLeapMediaCodecPlayer::UpdateAudioTrackInfo(MLHandle FormatHandle, FTrackData& CurrentTrackData)
+{
+	UpdateCommonTrackInfo(FormatHandle, CurrentTrackData);
+
+	MLResult Result = MLMediaFormatGetKeyValueInt32(FormatHandle, MLMediaFormat_Key_Sample_Rate, &CurrentTrackData.SampleRate);
+	UE_CLOG(Result != MLResult_Ok, LogMagicLeapMediaCodec, Error, TEXT("MLMediaFormatGetKeyValueInt32(sample-rate) failed with error %s"), UTF8_TO_TCHAR(MLMediaResultGetString(Result)));
+	Result = MLMediaFormatGetKeyValueInt32(FormatHandle, MLMediaFormat_Key_Channel_Count, &CurrentTrackData.ChannelCount);
+	UE_CLOG(Result != MLResult_Ok, LogMagicLeapMediaCodec, Error, TEXT("MLMediaFormatGetKeyValueInt32(channel-count) failed with error %s"), UTF8_TO_TCHAR(MLMediaResultGetString(Result)));
+}
+
+void FMagicLeapMediaCodecPlayer::UpdateVideoTrackInfo(MLHandle FormatHandle, FTrackData& CurrentTrackData)
+{
+	UpdateCommonTrackInfo(FormatHandle, CurrentTrackData);
+
+	MLResult Result = MLMediaFormatGetKeyValueInt32(FormatHandle, MLMediaFormat_Key_Height, &CurrentTrackData.Height);
+	UE_CLOG(Result != MLResult_Ok, LogMagicLeapMediaCodec, Error, TEXT("MLMediaFormatGetKeyValueInt32(height) failed with error %s"), UTF8_TO_TCHAR(MLMediaResultGetString(Result)));
+	Result = MLMediaFormatGetKeyValueInt32(FormatHandle, MLMediaFormat_Key_Width, &CurrentTrackData.Width);
+	UE_CLOG(Result != MLResult_Ok, LogMagicLeapMediaCodec, Error, TEXT("MLMediaFormatGetKeyValueInt32(width) failed with error %s"), UTF8_TO_TCHAR(MLMediaResultGetString(Result)));
+	Result = MLMediaFormatGetKeyValueInt32(FormatHandle, MLMediaFormat_Key_Frame_Rate, &CurrentTrackData.FrameRate);
+	UE_CLOG(Result != MLResult_Ok, LogMagicLeapMediaCodec, Error, TEXT("MLMediaFormatGetKeyValueInt32(frame-rate) failed with error %s"), UTF8_TO_TCHAR(MLMediaResultGetString(Result)));
+}
+
+bool FMagicLeapMediaCodecPlayer::CreateMediaCodec(EMediaTrackType TrackType, MLHandle& CodecHandle)
+{
+	if (SelectedTrack[TrackType] == INDEX_NONE)
+	{
+		return false;
+	}
+
+	FTrackData& CurrentTrack = TrackInfo[TrackType][SelectedTrack[TrackType]];
+	MLHandle TrackFormatHandle = ML_INVALID_HANDLE;
+
+	// When Configuring the codec, the correct track format must be set or you get absolute garbo
+	MLResult Result = MLMediaExtractorGetTrackFormat(MediaExtractorHandle, CurrentTrack.TrackIndex, &TrackFormatHandle);
+	UE_CLOG(Result != MLResult_Ok, LogMagicLeapMediaCodec, Error, TEXT("MLMediaExtractorGetTrackFormat() failed with error %s"), UTF8_TO_TCHAR(MLMediaResultGetString(Result)));
+
+	Result = MLMediaExtractorSelectTrack(MediaExtractorHandle, CurrentTrack.TrackIndex);
+	UE_CLOG(Result != MLResult_Ok, LogMagicLeapMediaCodec, Error, TEXT("MLMediaExtractorSelectTrack() failed with error %s"), UTF8_TO_TCHAR(MLMediaResultGetString(Result)));
+
+	if (!MLHandleIsValid(CodecHandle))
+	{
+		Result = MLMediaCodecCreateCodec(MLMediaCodecCreation_ByType, MLMediaCodecType_Decoder, TCHAR_TO_UTF8(*CurrentTrack.MimeType), &CodecHandle);
+		if (Result != MLResult_Ok)
+		{
+			UE_LOG(LogMagicLeapMediaCodec, Error, TEXT("MLMediaCodecCreateCodec() failed with error %s"), UTF8_TO_TCHAR(MLMediaResultGetString(Result)));
+			return false;
+		}
+	}
+
+	if (TrackType == EMediaTrackType::Video)
+	{
+		Result = MLMediaCodecSetSurfaceHint(VideoCodecHandle, MLMediaCodecSurfaceHint::MLMediaCodecSurfaceHint_Hardware);
+		UE_CLOG(Result != MLResult_Ok, LogMagicLeapMediaCodec, Error, TEXT("MLMediaCodecSetSurfaceHint() failed with error %s"), UTF8_TO_TCHAR(MLMediaResultGetString(Result)));
+	}
+
+	// TODO: Handle crypto for DRM
+	Result = MLMediaCodecConfigure(CodecHandle, TrackFormatHandle, static_cast<MLHandle>(0));
+	if (Result != MLResult_Ok)
+	{
+		UE_LOG(LogMagicLeapMediaCodec, Error, TEXT("MLMediaCodecConfigure failed with error %s."), UTF8_TO_TCHAR(MLMediaResultGetString(Result)));
+		return false;
+	}
+
+	Result = MLMediaFormatDestroy(TrackFormatHandle);
+	UE_CLOG(Result != MLResult_Ok, LogMagicLeapMediaCodec, Error, TEXT("MLMediaFormatDestroy() failed with error %s"), UTF8_TO_TCHAR(MLMediaResultGetString(Result)));
+
+	return true;
+}
+
+void FMagicLeapMediaCodecPlayer::StartMediaCodec(MLHandle& CodecHandle)
+{
+	if (MLHandleIsValid(CodecHandle))
+	{
+		MLResult Result = MLMediaCodecStart(CodecHandle);
+		UE_CLOG(Result != MLResult_Ok, LogMagicLeapMediaCodec, Error, TEXT("MLMediaCodecStart() failed with error %s"), UTF8_TO_TCHAR(MLMediaResultGetString(Result)));
+
+		// Output streams are related to the codecs, so lets reset the flags here.
+		bPlaybackCompleted = false;
+		bReachedOutputEndOfStream = false;
+		CurrentPlaybackTime = FTimespan::Zero();
+	}
+}
+
+// TODO: Remove this, we know when a buffer is available
+bool FMagicLeapMediaCodecPlayer::IsBufferAvailable_RenderThread(MLHandle InMediaPlayerHandle)
+{
+	ensureMsgf(IsInRenderingThread(), TEXT("IsBufferAvailable_RenderThread called outside of render thread"));
+	return bIsBufferAvailable;
+}
+
+bool FMagicLeapMediaCodecPlayer::GetNativeBuffer_RenderThread(const MLHandle InVideoCodecHandle, MLHandle& NativeBuffer)
+{
+	ensureMsgf(IsInRenderingThread(), TEXT("GetNativeBuffer_RenderThread called outside of render thread"));
+	MLResult Result = MLMediaCodecAcquireNextAvailableFrame(InVideoCodecHandle, &NativeBuffer);
+	if (Result != MLResult_Ok)
+	{
+		UE_LOG(LogMagicLeapMediaCodec, Error, TEXT("MLMediaCodecAcquireNextAvailableFrame failed with error %s"), UTF8_TO_TCHAR(MLMediaResultGetString(Result)));
+		return false;
+	}
+	return true;
+}
+
+bool FMagicLeapMediaCodecPlayer::ReleaseNativeBuffer_RenderThread(const MLHandle InVideoCodecHandle, MLHandle NativeBuffer)
+{
+	ensureMsgf(IsInRenderingThread(), TEXT("ReleaseNativeBuffer_RenderThread called outside of render thread"));
+	MLResult Result = MLMediaCodecReleaseFrame(InVideoCodecHandle, NativeBuffer);
+	if (Result != MLResult_Ok)
+	{
+		UE_LOG(LogMagicLeapMediaCodec, Error, TEXT("MLMediaCodecReleaseFrame failed with error %s"), UTF8_TO_TCHAR(MLMediaResultGetString(Result)));
+		return false;
+	}
+	return true;
+}
+
+bool FMagicLeapMediaCodecPlayer::GetCurrentPosition_RenderThread(const MLHandle InVideoCodecHandle, int32& CurrentPosition)
+{
+	return true;
+}
+
+bool FMagicLeapMediaCodecPlayer::ProcessVideoOutputSample_RenderThread(MLHandle& CodecHandle, FTrackData& CurrentTrackData, const FTimespan& Timecode)
+{
+	// Doing this to support a thread safe FlushCodecs() function.
+	FScopeTryLock LockIT(&RT_IT_Mutex);
+	// No point in waiting here, blocking the entire render thread. This would be locked only during a flush operation which requires GameThread lock as well.
+	// When the flush succeeds, we won't immediately get output buffers. We can afford to come back on the next frame.
+	if (!LockIT.IsLocked())
+	{
+		return false;
+	}
+
+	if (!CurrentTrackData.bBufferPendingRender)
+	{
+		MLResult Result = MLMediaCodecDequeueOutputBuffer(CodecHandle, &CurrentTrackData.CurrentBufferInfo, 0, (int64_t*)&CurrentTrackData.CurrentBufferIndex);
+		if (Result != MLResult_Ok)
+		{
+			UE_LOG(LogMagicLeapMediaCodec, Error, TEXT("MLMediaCodecDequeueOutputBuffer failed with error %s"), UTF8_TO_TCHAR(MLMediaResultGetString(Result)));
+			return false;
+		}
+	}
+
+	FTimespan CurrentPresentationTime = FTimespan::FromMicroseconds(CurrentTrackData.CurrentBufferInfo.presentation_time_us);
+	if (0 <= CurrentTrackData.CurrentBufferIndex)
+	{
+		MLResult Result;
+
+		{
+			FScopeLock Lock(&CriticalSection);
+
+			if (bResetVideoCodecStartTime)
+			{
+				CurrentTrackData.StartPresentationTime = FTimespan::Zero();
+				bResetVideoCodecStartTime = false;
+			}
+
+			// Resets StartPresentationTime when video is looping.
+			// TODO: this might not run ever
+			if (bReachedOutputEndOfStream && !InputWorker.HasReachedInputEOS())
+			{
+				CurrentTrackData.StartPresentationTime = FTimespan::Zero();
+				bReachedOutputEndOfStream = false;
+			}
+		}
+
+		if (CurrentTrackData.StartPresentationTime.IsZero())
+		{
+			CurrentTrackData.StartPresentationTime = Timecode - CurrentPresentationTime;
+			UE_LOG(LogMagicLeapMediaCodec, Warning, TEXT("Reset start presentation time to %f. (%f - %f)"), CurrentTrackData.StartPresentationTime.GetTotalMilliseconds(), Timecode.GetTotalMilliseconds(), CurrentPresentationTime.GetTotalMilliseconds());
+		}
+
+		switch (MediaSourceType)
+		{
+			case EMediaSourceType::VideoOnly:
+			{
+				if ((CurrentTrackData.StartPresentationTime + CurrentPresentationTime - Timecode) > FTimespan::Zero())
+				{
+					CurrentTrackData.bBufferPendingRender = true;
+					return false;
+				}
+				break;
+			}
+			case EMediaSourceType::VideoAndAudio:
+			{
+				int32 SelectedAudioTrack = SelectedTrack[EMediaTrackType::Audio];
+				if (SelectedAudioTrack != INDEX_NONE)
+				{
+					const FTrackData& CurrentAudioTrack = TrackInfo[EMediaTrackType::Audio][SelectedAudioTrack];
+					const FTimespan LastAudioPresentationTime = FTimespan(FPlatformAtomics::AtomicRead(reinterpret_cast<const int64*>(&CurrentAudioTrack.LastPresentationTime)));
+					if ((CurrentPresentationTime - LastAudioPresentationTime) > FTimespan::FromMicroseconds(250))
+					{
+						// UE_LOG(LogMagicLeapMediaCodec, Warning, TEXT("Waiting on audio to catch up."));
+						CurrentTrackData.bBufferPendingRender = true;
+						// Reset StartPresentationTime to prevent video stuttering.
+						CurrentTrackData.StartPresentationTime = Timecode - CurrentPresentationTime;
+						return false;
+					}
+				}
+				if ((CurrentTrackData.StartPresentationTime + CurrentPresentationTime - Timecode) > FTimespan::Zero())
+				{
+					// UE_LOG(LogMagicLeapMediaCodec, Warning, TEXT("Wait for time to catchup. Start presentation time = %f, Timecode = %f, Presenation Time = %f"), CurrentTrackData.StartPresentationTime.GetTotalMilliseconds(), Timecode.GetTotalMilliseconds(), CurrentPresentationTime.GetTotalMilliseconds());
+					CurrentTrackData.bBufferPendingRender = true;
+					return false;
+				}
+				break;
+			}
+			default:
+				UE_LOG(LogMagicLeapMediaCodec, Warning, TEXT("What sorcery is this?"));
+		}
+
+		CurrentTrackData.bBufferPendingRender = false;
+		const FTimespan CachedCurrentPlaybackTime = CurrentPlaybackTime.Load();
+
+		Result = MLMediaCodecReleaseOutputBuffer(CodecHandle, (MLHandle)CurrentTrackData.CurrentBufferIndex, CurrentTrackData.CurrentBufferInfo.size != 0);
+		if (Result != MLResult_Ok)
+		{
+			UE_LOG(LogMagicLeapMediaCodec, Error, TEXT("MLMediaCodecReleaseOutputBuffer(video) failed with error %s"), UTF8_TO_TCHAR(MLMediaResultGetString(Result)));
+			return false;
+		}
+		if (CurrentPresentationTime >= GetDuration())
+		{
+			FScopeLock Lock(&CriticalSection);
+			bReachedOutputEndOfStream = true;
+			bPlaybackCompleted = true;
+			bPlaybackCompleted_RenderThread = true;
+			CurrentTrackData.StartPresentationTime = FTimespan::Zero();
+			UE_LOG(LogMagicLeapMediaCodec, Display, TEXT("Playback ended on stream %s"), *CurrentTrackData.MimeType);
+		}
+		else if (!CachedCurrentPlaybackTime.IsZero() && CurrentPresentationTime == CachedCurrentPlaybackTime)
+		{
+			FScopeLock Lock(&CriticalSection);
+			bReachedOutputEndOfStream = true;
+			bPlaybackCompleted = true;
+			bPlaybackCompleted_RenderThread = true;
+			CurrentTrackData.StartPresentationTime = FTimespan::Zero();
+			UE_LOG(LogMagicLeapMediaCodec, Display, TEXT("Playback ended on stream %s"), *CurrentTrackData.MimeType);
+		}
+
+		CurrentTrackData.LastPresentationTime = Timecode;
+		CurrentPlaybackTime = CurrentPresentationTime;
+
+		return true;
+	}
+	else if (CurrentTrackData.CurrentBufferIndex == MLMediaCodec_FormatChanged)
+	{
+		UE_LOG(LogMagicLeapMediaCodec, Display, TEXT("%s MLMediaCodec_FormatChanged"), *CurrentTrackData.MimeType);
+
+		MLHandle NewFormatHandle;
+		MLResult Result = MLMediaCodecGetOutputFormat(CodecHandle, &NewFormatHandle);
+		if (Result != MLResult_Ok)
+		{
+			UE_LOG(LogMagicLeapMediaCodec, Error, TEXT("MLMediaCodecGetOutputFormat() failed with error %s"), UTF8_TO_TCHAR(MLMediaResultGetString(Result)));
+			return false;
+		}
+
+		UpdateVideoTrackInfo(NewFormatHandle, CurrentTrackData);
+
+		Result = MLMediaFormatDestroy(NewFormatHandle);
+		UE_CLOG(Result != MLResult_Ok, LogMagicLeapMediaCodec, Error, TEXT("MLMediaFormatDestroy() failed with error %s"), UTF8_TO_TCHAR(MLMediaResultGetString(Result)));
+	}
+	else if (CurrentTrackData.CurrentBufferIndex == MLMediaCodec_TryAgainLater)
+	{
+		if (InputWorker.HasReachedInputEOS())
+		{
+			// FScopeLock Lock(&CriticalSection);
+			bReachedOutputEndOfStream = true;
+			bPlaybackCompleted = true;
+			bPlaybackCompleted_RenderThread = true;
+			CurrentTrackData.StartPresentationTime = FTimespan::Zero();
+			UE_LOG(LogMagicLeapMediaCodec, Warning, TEXT("Playback ended on stream %s"), *CurrentTrackData.MimeType);
+		}
+		else
+		{
+			// UE_LOG(LogMagicLeapMediaCodec, Warning, TEXT("no video buffers available"));
+		}
+	}
+	// MLMediaCodec_OutputBuffersChanged is deprecated and doesnt need to be handled.
+	// else if (CurrentTrackData.CurrentBufferIndex == MLMediaCodec_OutputBuffersChanged){}
+
+	return false;
+}
+
+bool FMagicLeapMediaCodecPlayer::ProcessAudioOutputSample(MLHandle& CodecHandle, FTrackData& CurrentTrackData, FTimespan LastAudioSampleTime)
+{
+	// Doing this to support a thread safe FlushCodecs() function.
+	// No point in waiting here, blocking the entire game thread. This would be locked only during a flush operation which requires RenderThread lock as well.
+	// When the flush succeeds, we won't immediately get output buffers. We can afford to come back on the next frame.
+	FScopeTryLock LockIT(&GT_IT_Mutex);
+	if (!LockIT.IsLocked())
+	{
+		return false;
+	}
+
+	// Atomic becuase this will be used on the render thread to determine playback.
+	FPlatformAtomics::InterlockedExchange(reinterpret_cast<int64*>(&CurrentTrackData.LastPresentationTime), LastAudioSampleTime.GetTicks());
+
+	if (!CurrentTrackData.bBufferPendingRender)
+	{
+		MLResult Result = MLMediaCodecDequeueOutputBuffer(CodecHandle, &CurrentTrackData.CurrentBufferInfo, 0, (int64_t*)&CurrentTrackData.CurrentBufferIndex);
+		if (Result != MLResult_Ok)
+		{
+			UE_LOG(LogMagicLeapMediaCodec, Error, TEXT("MLMediaCodecDequeueOutputBuffer failed with error %s"), UTF8_TO_TCHAR(MLMediaResultGetString(Result)));
+			return false;
+		}
+	}
+
+	FTimespan CurrentPresentationTime = FTimespan::FromMicroseconds(CurrentTrackData.CurrentBufferInfo.presentation_time_us);
+	if (0 <= CurrentTrackData.CurrentBufferIndex)
+	{
+		SIZE_T BufferSize = 0;
+		const uint8_t* Buffer = nullptr;
+		MLResult Result = MLMediaCodecGetOutputBufferPointer(CodecHandle, (MLHandle)CurrentTrackData.CurrentBufferIndex, &Buffer, (size_t*)(&BufferSize));
+
+		bool bAudioBufferWritten = WriteAudioSample(CurrentTrackData, Buffer, BufferSize);
+
+		Result = MLMediaCodecReleaseOutputBuffer(CodecHandle, CurrentTrackData.CurrentBufferIndex, false);
+		UE_CLOG(Result != MLResult_Ok, LogMagicLeapMediaCodec, Error, TEXT("MLMediaCodecReleaseOutputBuffer(audio) failed with error %s"), UTF8_TO_TCHAR(MLMediaResultGetString(Result)));
+
+		if (bAudioBufferWritten)
+		{
+			CurrentTrackData.LastSampleQueueTime = CurrentPresentationTime.GetTicks();
+		}
+		else
+		{
+			return false;
+		}
+
+		return true;
+	}
+	else if (CurrentTrackData.CurrentBufferIndex == MLMediaCodec_FormatChanged)
+	{
+		UE_LOG(LogMagicLeapMediaCodec, Display, TEXT("%s MLMediaCodec_FormatChanged"), *CurrentTrackData.MimeType);
+
+		MLHandle NewFormatHandle;
+		MLResult Result = MLMediaCodecGetOutputFormat(CodecHandle, &NewFormatHandle);
+		if (Result != MLResult_Ok)
+		{
+			UE_LOG(LogMagicLeapMediaCodec, Error, TEXT("MLMediaCodecGetOutputFormat() failed with error %s"), UTF8_TO_TCHAR(MLMediaResultGetString(Result)));
+			return false;
+		}
+
+		UpdateAudioTrackInfo(NewFormatHandle, CurrentTrackData);
+
+		Result = MLMediaFormatDestroy(NewFormatHandle);
+		UE_CLOG(Result != MLResult_Ok, LogMagicLeapMediaCodec, Error, TEXT("MLMediaFormatDestroy() failed with error %s"), UTF8_TO_TCHAR(MLMediaResultGetString(Result)));
+	}
+	else if (CurrentTrackData.CurrentBufferIndex == MLMediaCodec_TryAgainLater)
+	{}
+	// MLMediaCodec_OutputBuffersChanged is deprecated and doesnt need to be handled.
+	// else if (CurrentTrackData.CurrentBufferIndex == MLMediaCodec_OutputBuffersChanged){}
+
+	return false;
+}
+
+bool FMagicLeapMediaCodecPlayer::WriteAudioSample(FTrackData& CurrentTrackData, const uint8* SampleBuffer, uint64 SampleSize)
+{
+	if ((SampleSize < 0) || (SampleBuffer == nullptr && SampleSize != 0))
+	{
+		return false;
+	}
+
+	CurrentTrackData.SampleDuration = (SampleSize * ETimespan::TicksPerSecond) / (CurrentTrackData.ChannelCount * CurrentTrackData.SampleRate * sizeof(int16));
+
+	// create & add sample to queue
+	const TSharedRef<FMagicLeapMediaAudioSample, ESPMode::ThreadSafe> AudioSample = AudioSamplePool->AcquireShared();
+
+	if (AudioSample->Initialize(SampleBuffer, SampleSize, CurrentTrackData.ChannelCount, CurrentTrackData.SampleRate, FTimespan::FromMicroseconds(CurrentTrackData.CurrentBufferInfo.presentation_time_us), CurrentTrackData.SampleDuration))
+	{
+		Samples->AddAudio(AudioSample);
+		return true;
+	}
+
+	return false;
+}
+
+#undef LOCTEXT_NAMESPACE
diff --git a/Engine/Plugins/Lumin/MagicLeapMedia/Source/MagicLeapMediaCodec/Private/MagicLeapMediaCodecPlayer.h b/Engine/Plugins/Lumin/MagicLeapMedia/Source/MagicLeapMediaCodec/Private/MagicLeapMediaCodecPlayer.h
new file mode 100644
--- /dev/null
+++ b/Engine/Plugins/Lumin/MagicLeapMedia/Source/MagicLeapMediaCodec/Private/MagicLeapMediaCodecPlayer.h
@@ -0,0 +1,266 @@
+// Copyright 1998-2018 Epic Games, Inc. All Rights Reserved.
+
+#pragma once
+
+#include "CoreMinimal.h"
+#include "Containers/UnrealString.h"
+#include "Misc/Guid.h"
+#include "IMediaCache.h"
+#include "IMediaControls.h"
+#include "IMediaPlayer.h"
+#include "IMediaTracks.h"
+#include "IMediaView.h"
+#include "IMediaEventSink.h"
+#include "IMediaOptions.h"
+#include "Misc/Timespan.h"
+#include "Templates/SharedPointer.h"
+#include "Templates/Atomic.h"
+#include "RenderingThread.h"
+#include "RHI.h"
+#include "RHIResources.h"
+#include "HAL/ThreadSafeCounter.h"
+#include "Containers/Queue.h"
+#include "Serialization/Archive.h"
+#include "Serialization/ArrayReader.h"
+#include "MediaCodecInputWorker.h"
+
+#ifndef EGL_EGLEXT_PROTOTYPES
+#define EGL_EGLEXT_PROTOTYPES
+#endif // !EGL_EGLEXT_PROTOTYPES
+
+#include <EGL/egl.h>
+#include <EGL/eglext.h>
+
+#ifndef GL_GLEXT_PROTOTYPES
+#define GL_GLEXT_PROTOTYPES
+#endif // !GL_GLEXT_PROTOTYPES
+
+#include <GLES2/gl2.h>
+#include <GLES2/gl2ext.h>
+
+#include "ml_api.h"
+#include "ml_media_codec.h"
+#include "ml_audio.h"
+#include "ml_media_data_source.h"
+
+class FMediaSamples;
+class IMediaEventSink;
+struct FMagicLeapVideoTextureData;
+class FMagicLeapMediaAudioSamplePool;
+
+/**
+*	Implement media playback using the MagicLeap MediaPlayer interface.
+*/
+class FMagicLeapMediaCodecPlayer : public IMediaPlayer, public IMediaControls, public IMediaCache, public IMediaTracks, public IMediaView
+{
+public:
+
+	/**
+	* Create and initialize a new instance.
+	*
+	* @param InEventSink The object that receives media events from this player.
+	*/
+	FMagicLeapMediaCodecPlayer(IMediaEventSink& InEventSink);
+
+	/** Virtual destructor. */
+	virtual ~FMagicLeapMediaCodecPlayer();
+
+public:
+	/** IMediaPlayer interface */
+	virtual void Close() override;
+	virtual IMediaCache& GetCache() override;
+	virtual IMediaControls& GetControls() override;
+	virtual FString GetInfo() const override;
+	virtual FName GetPlayerName() const override;
+	virtual IMediaSamples& GetSamples() override;
+	virtual FString GetStats() const override;
+	virtual IMediaTracks& GetTracks() override;
+	virtual FString GetUrl() const override;
+	virtual IMediaView& GetView() override;
+	virtual bool Open(const FString& Url, const IMediaOptions* Options) override;
+	virtual bool Open(const TSharedRef<FArchive, ESPMode::ThreadSafe>& Archive, const FString& OriginalUrl, const IMediaOptions* Options) override;
+	virtual void SetGuid(const FGuid& Guid) override;
+	virtual void TickFetch(FTimespan DeltaTime, FTimespan Timecode) override;
+	virtual void TickInput(FTimespan DeltaTime, FTimespan Timecode) override;
+	virtual void TickAudio(FTimespan LastAudioSampleTime) override;
+
+public:
+	/** IMediaControls interface */
+	virtual bool CanControl(EMediaControl Control) const override;
+	virtual FTimespan GetDuration() const override;
+	virtual float GetRate() const override;
+	virtual EMediaState GetState() const override;
+	virtual EMediaStatus GetStatus() const override;
+	virtual TRangeSet<float> GetSupportedRates(EMediaRateThinning Thinning) const override;
+	virtual FTimespan GetTime() const override;
+	virtual bool IsLooping() const override;
+	virtual bool Seek(const FTimespan& Time) override;
+	virtual bool SetLooping(bool Looping) override;
+	virtual bool SetRate(float Rate) override;
+	virtual bool SetNativeVolume(float Volume) override;
+
+public:
+	/** IMediaTracks interface */
+	virtual bool GetAudioTrackFormat(int32 TrackIndex, int32 FormatIndex, FMediaAudioTrackFormat& OutFormat) const override;
+	virtual int32 GetNumTracks(EMediaTrackType TrackType) const override;
+	virtual int32 GetNumTrackFormats(EMediaTrackType TrackType, int32 TrackIndex) const override;
+	virtual int32 GetSelectedTrack(EMediaTrackType TrackType) const override;
+	virtual FText GetTrackDisplayName(EMediaTrackType TrackType, int32 TrackIndex) const override;
+	virtual int32 GetTrackFormat(EMediaTrackType TrackType, int32 TrackIndex) const override;
+	virtual FString GetTrackLanguage(EMediaTrackType TrackType, int32 TrackIndex) const override;
+	virtual FString GetTrackName(EMediaTrackType TrackType, int32 TrackIndex) const override;
+	virtual bool GetVideoTrackFormat(int32 TrackIndex, int32 FormatIndex, FMediaVideoTrackFormat& OutFormat) const override;
+	virtual bool SelectTrack(EMediaTrackType TrackType, int32 TrackIndex) override;
+	virtual bool SetTrackFormat(EMediaTrackType TrackType, int32 TrackIndex, int32 FormatIndex) override;
+
+public:
+	void QueueMediaEvent(EMediaEvent MediaEvent);
+	void FlushCodecs();
+	bool GetCodecForTrackIndex(int64_t TrackIndex, MLHandle& SampleCodecHandle) const;
+	bool IsPlaybackCompleted() const;
+	void SetPlaybackCompleted(bool PlaybackCompleted);
+	void QueueVideoCodecStartTimeReset();
+
+	int64_t MediaDataSourceReadAt(MLHandle media_data_source, size_t position, size_t size, uint8_t *buffer);
+	int64_t MediaDataSourceGetSize(MLHandle media_data_source);
+	void MediaDataSourceClose(MLHandle media_data_source);
+
+private:
+	/** Callback for when the application resumed in the foreground. */
+	void HandleApplicationHasEnteredForeground();
+
+	/** Callback for when the application is being paused in the background. */
+	void HandleApplicationWillEnterBackground();
+
+	/** Track information */
+	// TODO: Move or rename, get rid of height/width since this is used for both audio and video
+	// Perhaps an FAudioTrackData and FVideoTrackData inheriting from FTrackData?
+	// Also make the codec handle a part of this struct. Will make it easier to perform all operations.
+	struct FTrackData 
+	{
+		FTrackData(FString Mime, int32 Index) : MimeType(Mime), TrackIndex(Index) {}
+		FString MimeType;
+		int32 TrackIndex = 0;
+		int64 Position = 0;
+		FTimespan Duration = 0;
+		int32 Height = 0;
+		int32 Width = 0;
+		bool IsPlaying = false;
+		MLMediaCodecBufferInfo CurrentBufferInfo = {};
+		int64 CurrentBufferIndex = 0;
+		bool bBufferPendingRender = false;
+		FTimespan StartPresentationTime = FTimespan::Zero();
+		FTimespan LastPresentationTime = FTimespan::Zero();
+		FTimespan LastSampleQueueTime = FTimespan::Zero();
+		FTimespan SampleDuration = FTimespan::Zero();
+		int32 SampleRate = 0;
+		int32 ChannelCount = 0;
+		FString FormatName = "";
+		FString Language = "";
+		int32 FrameRate = 30;
+	};
+
+	enum EMediaSourceType
+	{
+		VideoOnly,
+		AudioOnly,
+		VideoAndAudio
+	};
+
+protected:
+	FIntPoint GetVideoDimensions() const;
+
+	MLHandle AudioCodecHandle;
+	MLHandle VideoCodecHandle;
+	MLHandle MediaExtractorHandle;
+	MLHandle MediaDataSourceHandle;
+
+	bool bMediaPrepared;
+
+	/** Frame UV scale and offsets */
+	float UScale = 1.0f;
+	float VScale = 1.0f;
+	float UOffset = 0.0f;
+	float VOffset = 0.0f;
+
+	/** Frame transformation matrix used to determine UV offset and scaling */
+	float FrameTransformationMatrix[16] = {};
+
+	/** Current player state. */
+	EMediaState CurrentState;
+
+	/** The media event handler. */
+	IMediaEventSink& EventSink;
+
+	/** Media information string. */
+	FString Info;
+
+	/** Currently opened media. */
+	FString MediaUrl;
+
+	/** Media player Guid */
+	FGuid PlayerGuid;
+
+	/** Foreground/background delegate for pause. */
+	FDelegateHandle PauseHandle;
+
+	/** Foreground/background delegate for resume. */
+	FDelegateHandle ResumeHandle;
+
+	/** The media sample queue. */
+	TSharedPtr<FMediaSamples, ESPMode::ThreadSafe> Samples;
+
+	/** Audio sample object pool. */
+	FMagicLeapMediaAudioSamplePool* AudioSamplePool;
+
+	TSharedPtr<FMagicLeapVideoTextureData, ESPMode::ThreadSafe> TextureData;
+
+	TMap<EMediaTrackType, TArray<FTrackData>> TrackInfo;
+	TMap<EMediaTrackType, int32> SelectedTrack;
+
+	// Used for bPlaybackCompleted, CurrentState, bLoopPlayback, bResetVideoCodecStartTime, bReachedInputEndOfStream
+	FCriticalSection CriticalSection;
+	// Lock between Game thread and Input worker thread, used for synchronizing the codec flush operation.
+	FCriticalSection GT_IT_Mutex;
+	// Lock between Render thread and Input worker thread, used for synchronizing the codec flush operation.
+	FCriticalSection RT_IT_Mutex;
+
+	bool bWasMediaPlayingBeforeAppPause;
+	TAtomic<FTimespan> CurrentPlaybackTime;
+	EMediaSourceType MediaSourceType;
+
+	bool bLoopPlayback;
+	bool bPlaybackCompleted;
+	// Render thread cache of bPlaybackCompleted, used to clear out the VideoTexturePool for Vulkan without requiring a lock.
+	bool bPlaybackCompleted_RenderThread;
+	bool bIsBufferAvailable;
+	bool bReachedOutputEndOfStream;
+	// TODO: move to track.
+	bool bResetVideoCodecStartTime;
+
+	TQueue<EMediaEvent, EQueueMode::Spsc> MediaEventQueue;
+
+	FMediaCodecInputWorker InputWorker;
+
+	TSharedPtr<FArchive, ESPMode::ThreadSafe> DataSourceArchive;
+
+private:
+	// Non-Interface functions
+	virtual void GetTrackInformation();
+	void UpdateCommonTrackInfo(MLHandle FormatHandle, FTrackData& CurrentTrackData);
+	void UpdateAudioTrackInfo(MLHandle FormatHandle, FTrackData& CurrentTrackData);
+	void UpdateVideoTrackInfo(MLHandle FormatHandle, FTrackData& CurrentTrackData);
+	virtual bool StopAndReset(EMediaTrackType TrackType, MLHandle& CodecHandle);
+	virtual bool SetRateOne();
+	bool CreateMediaCodec(EMediaTrackType TrackType, MLHandle& CodecHandle);
+	void StartMediaCodec(MLHandle& CodecHandle);
+	bool UpdateTransformMatrix_RenderThread(MLHandle InVideoCodecHandle);
+	virtual void RegisterExternalTexture_RenderThread(const FGuid& InGuid, FTextureRHIRef& InTextureRHI, FSamplerStateRHIRef& InSamplerStateRHI);
+	virtual bool IsBufferAvailable_RenderThread(MLHandle MediaCodecHandle);
+	virtual bool GetNativeBuffer_RenderThread(const MLHandle InMediaCodecHandle, MLHandle& NativeBuffer);
+	virtual bool ReleaseNativeBuffer_RenderThread(const MLHandle InMediaCodecHandle, MLHandle NativeBuffer);
+	virtual bool GetCurrentPosition_RenderThread(const MLHandle InMediaCodecHandle, int32& CurrentPosition);
+	bool ProcessAudioOutputSample(MLHandle& CodecHandle, FTrackData& CurrentTrackData, FTimespan LastAudioSampleTime);
+	bool ProcessVideoOutputSample_RenderThread(MLHandle& CodecHandle, FTrackData& CurrentTrackData, const FTimespan& Timecode);
+	bool WriteAudioSample(FTrackData& CurrentTrackData, const uint8* SampleBuffer, uint64 SampleSize);
+};
diff --git a/Engine/Plugins/Lumin/MagicLeapMedia/Source/MagicLeapMediaCodec/Private/MediaCodecInputWorker.cpp b/Engine/Plugins/Lumin/MagicLeapMedia/Source/MagicLeapMediaCodec/Private/MediaCodecInputWorker.cpp
new file mode 100644
--- /dev/null
+++ b/Engine/Plugins/Lumin/MagicLeapMedia/Source/MagicLeapMediaCodec/Private/MediaCodecInputWorker.cpp
@@ -0,0 +1,251 @@
+// Copyright 1998-2018 Epic Games, Inc. All Rights Reserved.
+
+#include "MediaCodecInputWorker.h"
+#include "Misc/ScopeLock.h"
+#include "HAL/RunnableThread.h"
+#include "HAL/Event.h"
+#include "MagicLeapMediaCodecPlayer.h"
+#include "Lumin/LuminAffinity.h"
+#include "IMagicLeapMediaCodecModule.h"
+
+#include "ml_media_extractor.h"
+#include "ml_media_codec.h"
+#include "ml_media_error.h"
+
+FMediaCodecInputWorker::FMediaCodecInputWorker()
+: OwnerPlayer(nullptr)
+, ExtractorHandle(nullptr)
+, CriticalSection(nullptr)
+, GT_IT_Mutex(nullptr)
+, RT_IT_Mutex(nullptr)
+, Thread(nullptr)
+, Semaphore(nullptr)
+, StopTaskCounter(0)
+, bReachedInputEndOfStream(false)
+{}
+
+FMediaCodecInputWorker::~FMediaCodecInputWorker()
+{
+	DestroyThread();
+}
+
+void FMediaCodecInputWorker::InitThread(FMagicLeapMediaCodecPlayer& InOwnerPlayer, MLHandle& InExtractorHandle, FCriticalSection& InCriticalSection, FCriticalSection& InGT_IT_Mutex, FCriticalSection& InRT_IT_Mutex)
+{
+	OwnerPlayer = &InOwnerPlayer;
+	ExtractorHandle = &InExtractorHandle;
+	CriticalSection = &InCriticalSection;
+	GT_IT_Mutex = &InGT_IT_Mutex;
+	RT_IT_Mutex = &InRT_IT_Mutex;
+	StopTaskCounter.Set(0);
+	bReachedInputEndOfStream = false;
+
+	if (Semaphore == nullptr)
+	{
+		Semaphore = FGenericPlatformProcess::GetSynchEventFromPool(false);
+		Thread = FRunnableThread::Create(this, TEXT("MLMediaCodecInputWorker"), 0, TPri_Normal, FLuminAffinity::GetPoolThreadMask());
+	}
+}
+
+void FMediaCodecInputWorker::DestroyThread()
+{
+	StopTaskCounter.Increment();
+	if (Semaphore != nullptr)
+	{
+		Semaphore->Trigger();
+		Thread->WaitForCompletion();
+		FGenericPlatformProcess::ReturnSynchEventToPool(Semaphore);
+		Semaphore = nullptr;
+		delete Thread;
+		Thread = nullptr;		
+	}
+}
+
+uint32 FMediaCodecInputWorker::Run()
+{
+	FInputWorkerTask QueuedTask;
+
+	while (StopTaskCounter.GetValue() == 0)
+	{
+		EMediaState CurrentState = EMediaState::Error;
+
+		{
+			FScopeLock Lock(CriticalSection);
+			CurrentState = OwnerPlayer->GetState();
+		}
+
+		if (CurrentState != EMediaState::Playing && CurrentState != EMediaState::Preparing)
+		{
+			Semaphore->Wait();
+		}
+
+		{
+			FScopeLock Lock(CriticalSection);
+			if (OwnerPlayer->IsPlaybackCompleted() && OwnerPlayer->IsLooping())
+			{
+				bReachedInputEndOfStream = false;
+				OwnerPlayer->SetPlaybackCompleted(false);
+				// Enqueue on the incoming tasks queue instead of directly calling Seek_WorkerThread so that CriticalSection mutex is not locked by this worker thread while seeking.
+				Seek(FTimespan::Zero());
+			}
+		}
+
+		if (IncomingTasks.Dequeue(QueuedTask))
+		{
+			if (QueuedTask.TaskType == EInputWorkerTaskType::Seek)
+			{
+				FScopeLock LockGT(GT_IT_Mutex);
+				FScopeLock LockRT(RT_IT_Mutex);
+				Seek_WorkerThread(QueuedTask.SeekTime);
+			}
+		}
+
+		ProcessInputSample_WorkerThread();
+
+		// ~120Hz
+		// TODO: make configurable?
+		// TODO: dont sleep if there is something in the the IncomingTasks queue.
+		FPlatformProcess::Sleep(0.008f);
+	}
+
+	return 0;
+}
+
+void FMediaCodecInputWorker::WakeUp()
+{
+	if (Semaphore != nullptr)
+	{
+		Semaphore->Trigger();
+	}
+}
+
+void FMediaCodecInputWorker::Seek(FTimespan SeekTime)
+{
+	IncomingTasks.Enqueue(FInputWorkerTask(EInputWorkerTaskType::Seek, SeekTime));
+}
+
+bool FMediaCodecInputWorker::HasReachedInputEOS() const
+{
+	return bReachedInputEndOfStream;
+}
+
+void FMediaCodecInputWorker::ProcessInputSample_WorkerThread()
+{
+	{
+		FScopeLock Lock(CriticalSection);
+		if (bReachedInputEndOfStream)
+		{
+			return;
+		}
+	}
+
+	int64 BufferIndex = MLMediaCodec_TryAgainLater;
+	SIZE_T BufferSize = 0;
+	uint8_t *Buffer = nullptr;
+	int64 SampleSize = -1;
+	int64 PresentationTimeUs = -1;
+	int64 TrackIndex = -1;
+	MLResult Result;
+	MLHandle SampleCodecHandle = ML_INVALID_HANDLE;
+
+	// loop until we find a vald track index
+	do
+	{
+		Result = MLMediaExtractorGetSampleTrackIndex(*ExtractorHandle, reinterpret_cast<int64_t*>(&TrackIndex));
+		UE_CLOG(Result != MLResult_Ok, LogMagicLeapMediaCodec, Error, TEXT("MLMediaExtractorGetSampleTrackIndex() failed with error %s"), UTF8_TO_TCHAR(MLMediaResultGetString(Result)));
+
+		if (TrackIndex >= 0)
+		{
+			if (!OwnerPlayer->GetCodecForTrackIndex(TrackIndex, SampleCodecHandle))
+			{
+				Result = MLMediaExtractorAdvance(*ExtractorHandle);
+				UE_CLOG(Result != MLResult_Ok, LogMagicLeapMediaCodec, Error, TEXT("MLMediaExtractorAdvance(audio track) failed with error %s"), UTF8_TO_TCHAR(MLMediaResultGetString(Result)));
+			}
+		}
+		else
+		{
+			UE_LOG(LogMagicLeapMediaCodec, Display, TEXT("negative track index from MLMediaExtractorGetSampleTrackIndex. Reached input EOS."));
+		
+			{
+				FScopeLock Lock(CriticalSection);
+				bReachedInputEndOfStream = true;
+			}
+
+			return;
+		}
+	} while (!MLHandleIsValid(SampleCodecHandle));
+
+	Result = MLMediaCodecDequeueInputBuffer(SampleCodecHandle, 0, reinterpret_cast<int64_t*>(&BufferIndex));
+	if (Result != MLResult_Ok)
+	{
+		UE_LOG(LogMagicLeapMediaCodec, Error, TEXT("MLMediaCodecDequeueInputBuffer failed with error %s"), UTF8_TO_TCHAR(MLMediaResultGetString(Result)));
+		return;
+	}
+	if (BufferIndex == MLMediaCodec_TryAgainLater)
+	{
+		// UE_LOG(LogMagicLeapMediaCodec, Warning, TEXT("no input buffers available"));
+		return;
+	}
+
+	Result = MLMediaCodecGetInputBufferPointer(SampleCodecHandle, BufferIndex, &Buffer, reinterpret_cast<size_t*>(&BufferSize));
+	if (Result != MLResult_Ok)
+	{
+		UE_LOG(LogMagicLeapMediaCodec, Error, TEXT("MLMediaCodecGetInputBufferPointer failed with error %s"), UTF8_TO_TCHAR(MLMediaResultGetString(Result)));
+		return;
+	}
+	else if (Buffer == nullptr)
+	{
+		UE_LOG(LogMagicLeapMediaCodec, Error, TEXT("Got a null buffer pointer from MLMediaCodecGetInputBufferPointer"));
+		return;
+	}
+
+	Result = MLMediaExtractorReadSampleData(*ExtractorHandle, Buffer, BufferSize, 0, reinterpret_cast<int64_t*>(&SampleSize));
+	if (Result != MLResult_Ok)
+	{
+		UE_LOG(LogMagicLeapMediaCodec, Error, TEXT("MLMediaExtractorReadSampleData failed with error %s"), UTF8_TO_TCHAR(MLMediaResultGetString(Result)));
+		return;
+	}
+	if (SampleSize < 0)
+	{
+		SampleSize = 0;
+		UE_LOG(LogMagicLeapMediaCodec, Display, TEXT("negative sample size from MLMediaExtractorReadSampleData. Reached input EOS."));
+
+		{
+			FScopeLock Lock(CriticalSection);
+			bReachedInputEndOfStream = true;
+		}
+	}
+
+	Result = MLMediaExtractorGetSampleTime(*ExtractorHandle, reinterpret_cast<int64_t*>(&PresentationTimeUs));
+	UE_CLOG(Result != MLResult_Ok, LogMagicLeapMediaCodec, Error, TEXT("MLMediaExtractorGetSampleTime() failed with error %s"), UTF8_TO_TCHAR(MLMediaResultGetString(Result)));
+
+	Result = MLMediaCodecQueueInputBuffer(SampleCodecHandle, (MLHandle)BufferIndex, 0, SampleSize, PresentationTimeUs, bReachedInputEndOfStream ? MLMediaCodecBufferFlag_EOS : 0);
+	if (Result != MLResult_Ok)
+	{
+		UE_LOG(LogMagicLeapMediaCodec, Error, TEXT("MLMediaCodecQueueInputBuffer failed with error %s"), UTF8_TO_TCHAR(MLMediaResultGetString(Result)));
+		return;
+	}
+
+	Result = MLMediaExtractorAdvance(*ExtractorHandle);
+	UE_CLOG(Result != MLResult_Ok, LogMagicLeapMediaCodec, Error, TEXT("MLMediaExtractorAdvance() failed with error %s"), UTF8_TO_TCHAR(MLMediaResultGetString(Result)));
+}
+
+bool FMediaCodecInputWorker::Seek_WorkerThread(const FTimespan& SeekTime)
+{
+	UE_LOG(LogMagicLeapMediaCodec, Display, TEXT("Seek(%f)"), SeekTime.GetTotalMilliseconds());
+	MLResult Result = MLMediaExtractorSeekTo(*ExtractorHandle, SeekTime.GetTotalMicroseconds(), MLMediaSeekMode_Closest_Sync);
+	if (Result == MLResult_Ok)
+	{
+		// TODO: Make thread safe.
+		// Flush codec samples after seeking
+		OwnerPlayer->FlushCodecs();
+		OwnerPlayer->QueueMediaEvent(EMediaEvent::SeekCompleted);
+		OwnerPlayer->QueueVideoCodecStartTimeReset();
+
+		return true;
+	}
+	else
+	{
+		UE_LOG(LogMagicLeapMediaCodec, Error, TEXT("MLMediaExtractorSeekTo failed with error %s"), UTF8_TO_TCHAR(MLGetResultString(Result)));
+		return false;
+	}
+}
diff --git a/Engine/Plugins/Lumin/MagicLeapMedia/Source/MagicLeapMediaCodec/Private/MediaCodecInputWorker.h b/Engine/Plugins/Lumin/MagicLeapMedia/Source/MagicLeapMediaCodec/Private/MediaCodecInputWorker.h
new file mode 100644
--- /dev/null
+++ b/Engine/Plugins/Lumin/MagicLeapMedia/Source/MagicLeapMediaCodec/Private/MediaCodecInputWorker.h
@@ -0,0 +1,70 @@
+// Copyright 1998-2018 Epic Games, Inc. All Rights Reserved.
+
+#pragma once
+
+#include "CoreMinimal.h"
+#include "HAL/Runnable.h"
+#include "HAL/ThreadSafeCounter.h"
+#include "Containers/Queue.h"
+
+#include "ml_api.h"
+
+class FMagicLeapMediaCodecPlayer;
+
+enum class EInputWorkerTaskType : uint8
+{
+	None,
+	Seek
+};
+
+struct FInputWorkerTask
+{
+public:
+	EInputWorkerTaskType TaskType;
+	FTimespan SeekTime;
+
+public:
+	FInputWorkerTask()
+	: TaskType(EInputWorkerTaskType::None)
+	, SeekTime(FTimespan::Zero())
+	{}
+
+	FInputWorkerTask(EInputWorkerTaskType InTaskType, const FTimespan& InSeekTime)
+	: TaskType(InTaskType)
+	, SeekTime(InSeekTime)
+	{}
+};
+
+class FMediaCodecInputWorker : public FRunnable
+{
+public:
+	FMediaCodecInputWorker();
+	virtual ~FMediaCodecInputWorker();
+
+	void InitThread(FMagicLeapMediaCodecPlayer& InOwnerPlayer, MLHandle& InExtractorHandle, FCriticalSection& InCriticalSection, FCriticalSection& InGT_IT_Mutex, FCriticalSection& InRT_IT_Mutex);
+	void DestroyThread();
+	virtual uint32 Run() override;
+
+	void WakeUp();
+	void Seek(FTimespan SeekTime);
+
+	bool HasReachedInputEOS() const;
+
+private:
+	void ProcessInputSample_WorkerThread();
+	bool Seek_WorkerThread(const FTimespan& SeekTime);
+
+	FMagicLeapMediaCodecPlayer* OwnerPlayer;
+	MLHandle* ExtractorHandle;
+	FCriticalSection* CriticalSection;
+	FCriticalSection* GT_IT_Mutex;
+	FCriticalSection* RT_IT_Mutex;
+
+	FRunnableThread* Thread;
+	FEvent* Semaphore;
+	FThreadSafeCounter StopTaskCounter;
+
+	bool bReachedInputEndOfStream;
+
+	TQueue<FInputWorkerTask, EQueueMode::Spsc> IncomingTasks;
+};
diff --git a/Engine/Plugins/Lumin/MagicLeapMedia/Source/MagicLeapMediaCodec/Public/IMagicLeapMediaCodecModule.h b/Engine/Plugins/Lumin/MagicLeapMedia/Source/MagicLeapMediaCodec/Public/IMagicLeapMediaCodecModule.h
new file mode 100644
--- /dev/null
+++ b/Engine/Plugins/Lumin/MagicLeapMedia/Source/MagicLeapMediaCodec/Public/IMagicLeapMediaCodecModule.h
@@ -0,0 +1,29 @@
+// Copyright 1998-2018 Epic Games, Inc. All Rights Reserved.
+
+#pragma once
+
+#include "CoreTypes.h"
+#include "Templates/SharedPointer.h"
+#include "Modules/ModuleInterface.h"
+
+class IMediaPlayer;
+
+/**
+ * Interface for the MagicLeapMediaCodec module.
+ */
+class IMagicLeapMediaCodecModule : public IModuleInterface
+{
+public:
+	/**
+	 * Creates a MagicLeap codec based media player.
+	 *
+	 * @return A new media player, or nullptr if a player couldn't be created.
+	 */
+	virtual TSharedPtr<IMediaPlayer, ESPMode::ThreadSafe> CreatePlayer(class IMediaEventSink& EventSink) = 0;
+
+public:
+	/** Virtual destructor. */
+	virtual ~IMagicLeapMediaCodecModule() { }
+};
+
+DECLARE_LOG_CATEGORY_EXTERN(LogMagicLeapMediaCodec, All, All);
diff --git a/Engine/Plugins/Lumin/MagicLeapMedia/Source/MagicLeapMediaCodecFactory/MagicLeapMediaCodecFactory.Build.cs b/Engine/Plugins/Lumin/MagicLeapMedia/Source/MagicLeapMediaCodecFactory/MagicLeapMediaCodecFactory.Build.cs
new file mode 100644
--- /dev/null
+++ b/Engine/Plugins/Lumin/MagicLeapMedia/Source/MagicLeapMediaCodecFactory/MagicLeapMediaCodecFactory.Build.cs
@@ -0,0 +1,38 @@
+// Copyright 1998-2018 Epic Games, Inc. All Rights Reserved.
+
+namespace UnrealBuildTool.Rules
+{
+	public class MagicLeapMediaCodecFactory : ModuleRules
+	{
+		public MagicLeapMediaCodecFactory(ReadOnlyTargetRules Target) : base(Target)
+		{
+			DynamicallyLoadedModuleNames.AddRange(
+				new string[] {
+					"Media"
+				});
+
+			PrivateDependencyModuleNames.AddRange(
+				new string[] {
+					"Core",
+					"CoreUObject",
+					"MediaAssets"
+				});
+
+			PrivateIncludePathModuleNames.AddRange(
+				new string[] {
+					"MagicLeapMediaCodec",
+					"Media"
+				});
+
+			PrivateIncludePaths.AddRange(
+				new string[] {
+					"MagicLeapMediaFactory/Private"
+				});
+
+			if (Target.Platform == UnrealTargetPlatform.Lumin)
+			{
+				DynamicallyLoadedModuleNames.Add("MagicLeapMediaCodec");
+			}
+		}
+	}
+}
diff --git a/Engine/Plugins/Lumin/MagicLeapMedia/Source/MagicLeapMediaCodecFactory/Private/MagicLeapMediaCodecFactoryModule.cpp b/Engine/Plugins/Lumin/MagicLeapMedia/Source/MagicLeapMediaCodecFactory/Private/MagicLeapMediaCodecFactoryModule.cpp
new file mode 100644
--- /dev/null
+++ b/Engine/Plugins/Lumin/MagicLeapMedia/Source/MagicLeapMediaCodecFactory/Private/MagicLeapMediaCodecFactoryModule.cpp
@@ -0,0 +1,167 @@
+// Copyright 1998-2018 Epic Games, Inc. All Rights Reserved.
+
+#include "CoreMinimal.h"
+#include "Misc/Paths.h"
+#include "Modules/ModuleInterface.h"
+#include "Modules/ModuleManager.h"
+#include "IMediaModule.h"
+#include "IMediaOptions.h"
+#include "IMagicLeapMediaCodecModule.h"
+#include "MagicLeapMediaCodecFactoryPrivate.h"
+#include "IMediaPlayerFactory.h"
+
+DEFINE_LOG_CATEGORY(LogMagicLeapMediaCodecFactory);
+
+#define LOCTEXT_NAMESPACE "FMagicLeapMediaCodecFactoryModule"
+
+/**
+ * Implements the MagicLeapMediaCodecFactory module.
+ */
+class FMagicLeapMediaCodecFactoryModule : public IMediaPlayerFactory, public IModuleInterface
+{
+public:
+
+	/** IMediaPlayerFactory interface */
+	virtual bool CanPlayUrl(const FString& Url, const IMediaOptions* Options, TArray<FText>* OutWarnings, TArray<FText>* OutErrors) const override
+	{
+		FString Scheme;
+		FString Location;
+
+		// check scheme
+		if (!Url.Split(TEXT("://"), &Scheme, &Location, ESearchCase::CaseSensitive))
+		{
+			if (OutErrors != nullptr)
+			{
+				OutErrors->Add(LOCTEXT("NoSchemeFound", "No URI scheme found"));
+			}
+
+			return false;
+		}
+
+		if (!SupportedUriSchemes.Contains(Scheme))
+		{
+			if (OutErrors != nullptr)
+			{
+				OutErrors->Add(FText::Format(LOCTEXT("SchemeNotSupported", "The URI scheme '{0}' is not supported"), FText::FromString(Scheme)));
+			}
+
+			return false;
+		}
+
+		// check file extension
+		if (Scheme == TEXT("file"))
+		{
+			const FString Extension = FPaths::GetExtension(Location, false);
+
+			if (!SupportedFileExtensions.Contains(Extension))
+			{
+				if (OutErrors != nullptr)
+				{
+					OutErrors->Add(FText::Format(LOCTEXT("ExtensionNotSupported", "The file extension '{0}' is not supported"), FText::FromString(Extension)));
+				}
+
+				return false;
+			}
+		}
+
+		// check options
+		if ((OutWarnings != nullptr) && (Options != nullptr))
+		{
+			if (Options->GetMediaOption("PrecacheFile", false) && (Scheme != TEXT("file")))
+			{
+				OutWarnings->Add(LOCTEXT("PrecachingNotSupported", "Precaching is supported for local files only"));
+			}
+		}
+
+		return true;
+	}
+
+	virtual TSharedPtr<IMediaPlayer, ESPMode::ThreadSafe> CreatePlayer(IMediaEventSink& EventSink) override
+	{
+		auto MagicLeapMediaModule = FModuleManager::LoadModulePtr<IMagicLeapMediaCodecModule>("MagicLeapMediaCodec");
+		return (MagicLeapMediaModule != nullptr) ? MagicLeapMediaModule->CreatePlayer(EventSink) : nullptr;
+	}
+
+	virtual FText GetDisplayName() const override
+	{
+		return LOCTEXT("MediaPlayerDisplayName", "MagicLeap Media Codec");
+	}
+
+	virtual FName GetPlayerName() const override
+	{
+		static FName PlayerName(TEXT("MagicLeapMediaCodec"));
+		return PlayerName;
+	}
+
+	virtual const TArray<FString>& GetSupportedPlatforms() const override
+	{
+		return SupportedPlatforms;
+	}
+
+	virtual bool SupportsFeature(EMediaFeature Feature) const
+	{
+		// TODO: add subtitles, closed captions etc when we add support for that.
+		return ((Feature == EMediaFeature::AudioTracks) ||
+			(Feature == EMediaFeature::AudioSamples) ||
+			(Feature == EMediaFeature::VideoSamples) ||
+			(Feature == EMediaFeature::VideoTracks));
+	}
+
+public:
+
+	/** IModuleInterface interface */
+	
+	virtual void StartupModule() override
+	{
+		// supported file extensions
+		SupportedFileExtensions.Add(TEXT("mp4"));
+		SupportedFileExtensions.Add(TEXT("3gpp"));
+		SupportedFileExtensions.Add(TEXT("aac"));
+		SupportedFileExtensions.Add(TEXT("m3u8"));
+
+    	// supported platforms
+    	SupportedPlatforms.Add(TEXT("Lumin"));
+
+		// supported schemes
+		SupportedUriSchemes.Add(TEXT("file"));
+		SupportedUriSchemes.Add(TEXT("http"));
+		SupportedUriSchemes.Add(TEXT("https"));
+		// Not supporting streaming right now.
+		// SupportedUriSchemes.Add(TEXT("httpd"));
+		// SupportedUriSchemes.Add(TEXT("mms"));
+
+		// register media player info
+		auto MediaModule = FModuleManager::LoadModulePtr<IMediaModule>("Media");
+
+		if (MediaModule != nullptr)
+		{
+			MediaModule->RegisterPlayerFactory(*this);
+		}
+	}
+
+	virtual void ShutdownModule() override
+	{
+		// unregister player factory
+		auto MediaModule = FModuleManager::GetModulePtr<IMediaModule>("Media");
+		
+		if (MediaModule != nullptr)
+		{
+			MediaModule->UnregisterPlayerFactory(*this);
+		}
+	}
+
+private:
+
+	/** List of supported media file types. */
+	TArray<FString> SupportedFileExtensions;
+	
+	/** List of platforms that the media player support. */
+	TArray<FString> SupportedPlatforms;
+	
+	/** List of supported URI schemes. */
+	TArray<FString> SupportedUriSchemes;
+};
+
+#undef LOCTEXT_NAMESPACE
+
+IMPLEMENT_MODULE(FMagicLeapMediaCodecFactoryModule, MagicLeapMediaCodecFactory);
diff --git a/Engine/Plugins/Lumin/MagicLeapMedia/Source/MagicLeapMediaCodecFactory/Private/MagicLeapMediaCodecFactoryPrivate.h b/Engine/Plugins/Lumin/MagicLeapMedia/Source/MagicLeapMediaCodecFactory/Private/MagicLeapMediaCodecFactoryPrivate.h
new file mode 100644
--- /dev/null
+++ b/Engine/Plugins/Lumin/MagicLeapMedia/Source/MagicLeapMediaCodecFactory/Private/MagicLeapMediaCodecFactoryPrivate.h
@@ -0,0 +1,7 @@
+// Copyright 1998-2018 Epic Games, Inc. All Rights Reserved.
+
+#pragma once
+
+#include "CoreMinimal.h"
+
+DECLARE_LOG_CATEGORY_EXTERN(LogMagicLeapMediaCodecFactory, Log, All);
diff --git a/Engine/Plugins/Lumin/MagicLeapMedia/Source/MagicLeapMediaFactory/Private/MagicLeapMediaFactoryModule.cpp b/Engine/Plugins/Lumin/MagicLeapMedia/Source/MagicLeapMediaFactory/Private/MagicLeapMediaFactoryModule.cpp
--- a/Engine/Plugins/Lumin/MagicLeapMedia/Source/MagicLeapMediaFactory/Private/MagicLeapMediaFactoryModule.cpp
+++ b/Engine/Plugins/Lumin/MagicLeapMedia/Source/MagicLeapMediaFactory/Private/MagicLeapMediaFactoryModule.cpp
@@ -109,8 +109,6 @@ public:
 
     	// supported platforms
     	SupportedPlatforms.Add(TEXT("Lumin"));
-    	// Hack until we get a separete ini platform for Lumin. Will not affect Android since this plugin is not built for it.
-    	SupportedPlatforms.Add(TEXT("Android"));
 
 		// supported schemes
 		SupportedUriSchemes.Add(TEXT("file"));
diff --git a/Engine/Plugins/Lumin/MagicLeapScreens/Source/MagicLeapScreens/Private/MagicLeapScreensComponent.cpp b/Engine/Plugins/Lumin/MagicLeapScreens/Source/MagicLeapScreens/Private/MagicLeapScreensComponent.cpp
new file mode 100644
--- /dev/null
+++ b/Engine/Plugins/Lumin/MagicLeapScreens/Source/MagicLeapScreens/Private/MagicLeapScreensComponent.cpp
@@ -0,0 +1,132 @@
+// Copyright 1998-2018 Epic Games, Inc. All Rights Reserved.
+
+#include "MagicLeapScreensComponent.h"
+#include "CoreMinimal.h"
+#include "Containers/Array.h"
+#include "RenderUtils.h"
+#include "MagicLeapPluginUtil.h"
+#include "MagicLeapHMD.h"
+#include "MagicLeapUtils.h"
+#include "IMagicLeapScreensPlugin.h"
+#include "MagicLeapScreensPlugin.h"
+#include "HeadMountedDisplayFunctionLibrary.h"
+#include "Misc/CoreDelegates.h"
+#include "MagicLeapScreensWorker.h"
+#include "MagicLeapScreensMsg.h"
+#if PLATFORM_LUMIN
+#include "Lumin/LuminAffinity.h"
+#endif // PLATFORM_LUMIN
+#include "MagicLeapPluginUtil.h" // for ML_INCLUDES_START/END
+
+#if WITH_MLSDK
+ML_INCLUDES_START
+#include <ml_image.h>
+#include <ml_screens.h>
+ML_INCLUDES_END
+#endif //WITH_MLSDK
+
+DEFINE_LOG_CATEGORY_STATIC(LogScreensComponent, Display, All);
+
+UScreensComponent::UScreensComponent()
+{
+	PrimaryComponentTick.TickGroup = TG_PrePhysics;
+	PrimaryComponentTick.bStartWithTickEnabled = true;
+	PrimaryComponentTick.bCanEverTick = true;
+}
+
+UScreensComponent::~UScreensComponent()
+{ }
+
+void UScreensComponent::TickComponent(float DeltaTime, enum ELevelTick TickType, FActorComponentTickFunction *ThisTickFunction)
+{
+	Super::TickComponent(DeltaTime, TickType, ThisTickFunction);
+
+	if (!FMagicLeapScreensPlugin::Impl->OutgoingMessages.IsEmpty())
+	{
+		FScreensMessage Msg;
+		FMagicLeapScreensPlugin::Impl->OutgoingMessages.Dequeue(Msg);
+
+		if (Msg.Type == EScreensMsgType::Request)
+		{
+			UE_LOG(LogScreensComponent, Error, TEXT("Unexpected EScreensMsgType::Request received from worker thread!"));
+		}
+		else if (Msg.Type == EScreensMsgType::Response)
+		{
+			switch (Msg.TaskType)
+			{
+			case EScreensTaskType::None: break;
+			case EScreensTaskType::GetHistory:
+			{
+				if (Msg.bSuccess)
+				{
+					GetWatchHistorySuccess.Broadcast(Msg.WatchHistory);
+				}
+				else
+				{
+					GetWatchHistoryFailure.Broadcast();
+				}
+			}
+			break;
+			case EScreensTaskType::AddToHistory:
+			{
+				if (Msg.WatchHistory.Num() > 0)
+				{
+					FScreensWatchHistoryEntry Entry = Msg.WatchHistory[0];
+					AddToWatchHistoryResult.Broadcast(Entry, Msg.bSuccess);
+				}
+				else
+				{
+					UE_LOG(LogScreensComponent, Error, TEXT("Unexpected empty watch history in an AddToHistory response from the worker thread"));
+				}
+			}
+			break;
+			case EScreensTaskType::UpdateEntry:
+			{
+				if (Msg.WatchHistory.Num() > 0)
+				{
+					FScreensWatchHistoryEntry Entry = Msg.WatchHistory[0];
+					UpdateWatchHistoryEntryResult.Broadcast(Entry, Msg.bSuccess);
+				}
+				else
+				{
+					UE_LOG(LogScreensComponent, Error, TEXT("Unexpected empty watch history in an UpdateEntry response from the worker thread"));
+				}
+			}
+			break;
+			}
+		}
+	}
+}
+
+void UScreensComponent::GetWatchHistoryAsync()
+{
+	FMagicLeapScreensPlugin::GetWatchHistoryEntriesAsync(TOptional<FScreensHistoryRequestResultDelegate>());
+}
+
+void UScreensComponent::AddWatchHistoryEntryAsync(const FScreensWatchHistoryEntry& WatchHistoryEntry)
+{
+	FMagicLeapScreensPlugin::AddToWatchHistoryAsync(WatchHistoryEntry, TOptional<FScreensEntryRequestResultDelegate>());
+}
+
+void UScreensComponent::UpdateWatchHistoryEntryAsync(const FScreensWatchHistoryEntry& WatchHistoryEntry)
+{
+	FMagicLeapScreensPlugin::UpdateWatchHistoryEntryAsync(WatchHistoryEntry, TOptional<FScreensEntryRequestResultDelegate>());
+}
+
+bool UScreensComponent::RemoveWatchHistoryEntry(const FScreenID& ID)
+{
+	return FMagicLeapScreensPlugin::RemoveWatchHistoryEntry(ID);
+}
+
+bool UScreensComponent::ClearWatchHistory()
+{
+	return FMagicLeapScreensPlugin::ClearWatchHistory();
+}
+
+bool UScreensComponent::GetScreensTransforms(TArray<FScreenTransform>& ScreensTransforms)
+{
+#if WITH_MLSDK
+	return FMagicLeapScreensPlugin::GetScreensTransforms(ScreensTransforms);
+#endif // WITH_MLSDK
+	return false;
+}
\ No newline at end of file
diff --git a/Engine/Plugins/Lumin/MagicLeapScreens/Source/MagicLeapScreens/Private/MagicLeapScreensFunctionLibrary.cpp b/Engine/Plugins/Lumin/MagicLeapScreens/Source/MagicLeapScreens/Private/MagicLeapScreensFunctionLibrary.cpp
new file mode 100644
--- /dev/null
+++ b/Engine/Plugins/Lumin/MagicLeapScreens/Source/MagicLeapScreens/Private/MagicLeapScreensFunctionLibrary.cpp
@@ -0,0 +1,51 @@
+// Copyright 1998-2018 Epic Games, Inc. All Rights Reserved.
+
+#include "MagicLeapScreensFunctionLibrary.h"
+#include "MagicLeapHMD.h"
+#include "MagicLeapUtils.h"
+#include "MagicLeapScreensWorker.h"
+#include "MagicLeapScreensMsg.h"
+#include "MagicLeapScreensPlugin.h"
+
+void UMagicLeapScreensFunctionLibrary::GetWatchHistoryAsync(const FScreensHistoryRequestResultDelegate& ResultDelegate)
+{
+#if WITH_MLSDK
+	FMagicLeapScreensPlugin::GetWatchHistoryEntriesAsync(TOptional<FScreensHistoryRequestResultDelegate>(ResultDelegate));
+#else
+	TArray<FScreensWatchHistoryEntry> ResultEntries;
+	ResultDelegate.ExecuteIfBound(false, ResultEntries);
+#endif // WITH_MLSDK
+}
+
+void UMagicLeapScreensFunctionLibrary::AddToWatchHistoryAsync(const FScreensWatchHistoryEntry& NewEntry, const FScreensEntryRequestResultDelegate& ResultDelegate)
+{
+#if WITH_MLSDK
+	FMagicLeapScreensPlugin::AddToWatchHistoryAsync(NewEntry, TOptional<FScreensEntryRequestResultDelegate>(ResultDelegate));
+#else
+	ResultDelegate.ExecuteIfBound(false, NewEntry);
+#endif // WITH_MLSDK
+}
+
+void UMagicLeapScreensFunctionLibrary::UpdateWatchHistoryEntryAsync(const FScreensWatchHistoryEntry& UpdateEntry, const FScreensEntryRequestResultDelegate& ResultDelegate)
+{
+#if WITH_MLSDK
+	FMagicLeapScreensPlugin::UpdateWatchHistoryEntryAsync(UpdateEntry, TOptional<FScreensEntryRequestResultDelegate>(ResultDelegate));
+#else
+	ResultDelegate.ExecuteIfBound(false, UpdateEntry);
+#endif // WITH_MLSDK
+}
+
+bool UMagicLeapScreensFunctionLibrary::RemoveWatchHistoryEntry(const FScreenID& ID)
+{
+	return FMagicLeapScreensPlugin::RemoveWatchHistoryEntry(ID);
+}
+
+bool UMagicLeapScreensFunctionLibrary::ClearWatchHistory()
+{
+	return FMagicLeapScreensPlugin::ClearWatchHistory();
+}
+
+bool UMagicLeapScreensFunctionLibrary::GetScreensTransforms(TArray<FScreenTransform>& ScreensTransforms)
+{
+	return FMagicLeapScreensPlugin::GetScreensTransforms(ScreensTransforms);
+}
