Cleanup a merged change to IMediaPlayer so it does not require changes to all media players.

diff --git a/Engine/Plugins/Lumin/MagicLeapMedia/Source/MagicLeapMediaCodec/Private/MagicLeapMediaCodecPlayer.cpp b/Engine/Plugins/Lumin/MagicLeapMedia/Source/MagicLeapMediaCodec/Private/MagicLeapMediaCodecPlayer.cpp
--- a/Engine/Plugins/Lumin/MagicLeapMedia/Source/MagicLeapMediaCodec/Private/MagicLeapMediaCodecPlayer.cpp
+++ b/Engine/Plugins/Lumin/MagicLeapMedia/Source/MagicLeapMediaCodec/Private/MagicLeapMediaCodecPlayer.cpp
@@ -1002,11 +1002,16 @@ void FMagicLeapMediaCodecPlayer::TickInput(FTimespan DeltaTime, FTimespan Timeco
 	}
 }
 
-void FMagicLeapMediaCodecPlayer::TickAudio(FTimespan LastAudioSampleTime)
+void FMagicLeapMediaCodecPlayer::SetLastAudioRenderedSampleTime(FTimespan SampleTime)
+{
+	LastAudioRenderedSampleTime = SampleTime;
+}
+
+void FMagicLeapMediaCodecPlayer::TickAudio()
 {
 	if (GetSelectedTrack(EMediaTrackType::Audio) != INDEX_NONE)
 	{
-		ProcessAudioOutputSample(AudioCodecHandle, TrackInfo[EMediaTrackType::Audio][SelectedTrack[EMediaTrackType::Audio]], LastAudioSampleTime);
+		ProcessAudioOutputSample(AudioCodecHandle, TrackInfo[EMediaTrackType::Audio][SelectedTrack[EMediaTrackType::Audio]], LastAudioRenderedSampleTime);
 	}
 }
 
diff --git a/Engine/Plugins/Lumin/MagicLeapMedia/Source/MagicLeapMediaCodec/Private/MagicLeapMediaCodecPlayer.h b/Engine/Plugins/Lumin/MagicLeapMedia/Source/MagicLeapMediaCodec/Private/MagicLeapMediaCodecPlayer.h
--- a/Engine/Plugins/Lumin/MagicLeapMedia/Source/MagicLeapMediaCodec/Private/MagicLeapMediaCodecPlayer.h
+++ b/Engine/Plugins/Lumin/MagicLeapMedia/Source/MagicLeapMediaCodec/Private/MagicLeapMediaCodecPlayer.h
@@ -82,7 +82,8 @@ public:
 	virtual void SetGuid(const FGuid& Guid) override;
 	virtual void TickFetch(FTimespan DeltaTime, FTimespan Timecode) override;
 	virtual void TickInput(FTimespan DeltaTime, FTimespan Timecode) override;
-	virtual void TickAudio(FTimespan LastAudioSampleTime) override;
+	virtual void TickAudio() override;
+	virtual void SetLastAudioRenderedSampleTime(FTimespan SampleTime) override;
 
 public:
 	/** IMediaControls interface */
@@ -227,6 +228,7 @@ protected:
 
 	bool bWasMediaPlayingBeforeAppPause;
 	TAtomic<FTimespan> CurrentPlaybackTime;
+	TAtomic<FTimespan> LastAudioRenderedSampleTime;
 	EMediaSourceType MediaSourceType;
 
 	bool bLoopPlayback;
diff --git a/Engine/Plugins/Media/AvfMedia/Source/AvfMedia/Private/Player/AvfMediaPlayer.cpp b/Engine/Plugins/Media/AvfMedia/Source/AvfMedia/Private/Player/AvfMediaPlayer.cpp
--- a/Engine/Plugins/Media/AvfMedia/Source/AvfMedia/Private/Player/AvfMediaPlayer.cpp
+++ b/Engine/Plugins/Media/AvfMedia/Source/AvfMedia/Private/Player/AvfMediaPlayer.cpp
@@ -508,7 +508,7 @@ bool FAvfMediaPlayer::Open(const TSharedRef<FArchive, ESPMode::ThreadSafe>& /*Ar
 }
 
 
-void FAvfMediaPlayer::TickAudio(FTimespan LastAudioSampleTime)
+void FAvfMediaPlayer::TickAudio()
 {
 	// NOP
 }
diff --git a/Engine/Plugins/Media/AvfMedia/Source/AvfMedia/Private/Player/AvfMediaPlayer.h b/Engine/Plugins/Media/AvfMedia/Source/AvfMedia/Private/Player/AvfMediaPlayer.h
--- a/Engine/Plugins/Media/AvfMedia/Source/AvfMedia/Private/Player/AvfMediaPlayer.h
+++ b/Engine/Plugins/Media/AvfMedia/Source/AvfMedia/Private/Player/AvfMediaPlayer.h
@@ -64,7 +64,7 @@ public:
 	virtual IMediaView& GetView() override;
 	virtual bool Open(const FString& Url, const IMediaOptions* Options) override;
 	virtual bool Open(const TSharedRef<FArchive, ESPMode::ThreadSafe>& Archive, const FString& OriginalUrl, const IMediaOptions* Options) override;
-	virtual void TickAudio(FTimespan LastAudioSampleTime) override;
+	virtual void TickAudio() override;
 	virtual void TickFetch(FTimespan DeltaTime, FTimespan Timecode) override;
 	virtual void TickInput(FTimespan DeltaTime, FTimespan Timecode) override;
 
diff --git a/Engine/Plugins/Media/MfMedia/Source/MfMedia/Private/Player/MfMediaPlayer.cpp b/Engine/Plugins/Media/MfMedia/Source/MfMedia/Private/Player/MfMediaPlayer.cpp
--- a/Engine/Plugins/Media/MfMedia/Source/MfMedia/Private/Player/MfMediaPlayer.cpp
+++ b/Engine/Plugins/Media/MfMedia/Source/MfMedia/Private/Player/MfMediaPlayer.cpp
@@ -193,7 +193,7 @@ bool FMfMediaPlayer::Open(const TSharedRef<FArchive, ESPMode::ThreadSafe>& Archi
 }
 
 
-void FMfMediaPlayer::TickAudio(FTimespan LastAudioSampleTime)
+void FMfMediaPlayer::TickAudio()
 {
 	if (CurrentState == EMediaState::Playing)
 	{
diff --git a/Engine/Plugins/Media/MfMedia/Source/MfMedia/Private/Player/MfMediaPlayer.h b/Engine/Plugins/Media/MfMedia/Source/MfMedia/Private/Player/MfMediaPlayer.h
--- a/Engine/Plugins/Media/MfMedia/Source/MfMedia/Private/Player/MfMediaPlayer.h
+++ b/Engine/Plugins/Media/MfMedia/Source/MfMedia/Private/Player/MfMediaPlayer.h
@@ -69,7 +69,7 @@ public:
 	virtual bool Open(const FString& Url, const IMediaOptions* Options) override;
 	virtual bool Open(const FString& Url, const IMediaOptions* Options, const FMediaPlayerOptions* PlayerOptions) override;
 	virtual bool Open(const TSharedRef<FArchive, ESPMode::ThreadSafe>& Archive, const FString& OriginalUrl, const IMediaOptions* Options) override;
-	virtual void TickAudio(FTimespan LastAudioSampleTime) override;
+	virtual void TickAudio() override;
 	virtual void TickFetch(FTimespan DeltaTime, FTimespan Timecode) override;
 	virtual void TickInput(FTimespan DeltaTime, FTimespan Timecode) override;
 
diff --git a/Engine/Source/Runtime/Media/Public/IMediaPlayer.h b/Engine/Source/Runtime/Media/Public/IMediaPlayer.h
--- a/Engine/Source/Runtime/Media/Public/IMediaPlayer.h
+++ b/Engine/Source/Runtime/Media/Public/IMediaPlayer.h
@@ -223,6 +223,16 @@ public:
 		return false;
 	}
 
+	/**
+	 * Notify player of last sample time of audio used.
+	 *
+	 * @param SampleTime The last audio sample dequeued by one of the audio sinks.
+	 */
+	virtual void SetLastAudioRenderedSampleTime(FTimespan SampleTime)
+	{
+		// override in child classes if supported
+	}
+
 	/**
 	 * Tick the player's audio related code.
 	 *
@@ -231,7 +241,7 @@ public:
 	 *
 	 * @see TickInput, TickFetch
 	 */
-	virtual void TickAudio(FTimespan LastAudioSampleTime)
+	virtual void TickAudio()
 	{
 		// override in child class if needed
 	}
diff --git a/Engine/Source/Runtime/MediaUtils/Private/MediaPlayerFacade.cpp b/Engine/Source/Runtime/MediaUtils/Private/MediaPlayerFacade.cpp
--- a/Engine/Source/Runtime/MediaUtils/Private/MediaPlayerFacade.cpp
+++ b/Engine/Source/Runtime/MediaUtils/Private/MediaPlayerFacade.cpp
@@ -1063,7 +1063,9 @@ void FMediaPlayerFacade::TickTickable()
 		return;
 	}
 
-	Player->TickAudio(LastAudioRenderedSampleTime.Load());
+	Player->SetLastAudioRenderedSampleTime(LastAudioRenderedSampleTime.Load());
+
+	Player->TickAudio();
 
 	// determine range of valid samples
 	TRange<FTimespan> AudioTimeRange;
