Copying //UE4/Partner-MagicLeap-Staging to Dev-VR (//UE4/Dev-VR)

diff --git a/Engine/Plugins/Lumin/MagicLeapLocation/Source/Private/MagicLeapLocationComponent.cpp b/Engine/Plugins/Lumin/MagicLeapLocation/Source/Private/MagicLeapLocationComponent.cpp
new file mode 100644
--- /dev/null
+++ b/Engine/Plugins/Lumin/MagicLeapLocation/Source/Private/MagicLeapLocationComponent.cpp
@@ -0,0 +1,34 @@
+// Copyright 1998-2019 Epic Games, Inc. All Rights Reserved.
+
+#include "MagicLeapLocationComponent.h"
+#include "MagicLeapLocationPlugin.h"
+
+UMagicLeapLocationComponent::UMagicLeapLocationComponent()
+{
+	PrimaryComponentTick.TickGroup = TG_PrePhysics;
+	PrimaryComponentTick.bStartWithTickEnabled = true;
+	PrimaryComponentTick.bCanEverTick = true;
+}
+
+UMagicLeapLocationComponent::~UMagicLeapLocationComponent()
+{ }
+
+bool UMagicLeapLocationComponent::GetLastCoarseLocation(FLocationData& OutLocation)
+{
+	return GET_MAGIC_LEAP_LOCATION_PLUGIN()->GetLastCoarseLocation(OutLocation);
+}
+
+bool UMagicLeapLocationComponent::GetLastCoarseLocationAsync()
+{
+	return GET_MAGIC_LEAP_LOCATION_PLUGIN()->GetLastCoarseLocationAsync(OnGotCoarseLocation);
+}
+
+bool UMagicLeapLocationComponent::GetLastCoarseLocationOnSphere(float InRadius, FVector& OutLocation)
+{
+	return GET_MAGIC_LEAP_LOCATION_PLUGIN()->GetLastCoarseLocationOnSphere(InRadius, OutLocation);
+}
+
+bool UMagicLeapLocationComponent::GetLastCoarseLocationOnSphereAsync(float InRadius)
+{
+	return GET_MAGIC_LEAP_LOCATION_PLUGIN()->GetLastCoarseLocationOnSphereAsync(OnGotCoarseLocationOnSphere, InRadius);
+}
diff --git a/Engine/Plugins/Lumin/MagicLeapLocation/Source/Private/MagicLeapLocationFunctionLibrary.cpp b/Engine/Plugins/Lumin/MagicLeapLocation/Source/Private/MagicLeapLocationFunctionLibrary.cpp
new file mode 100644
--- /dev/null
+++ b/Engine/Plugins/Lumin/MagicLeapLocation/Source/Private/MagicLeapLocationFunctionLibrary.cpp
@@ -0,0 +1,28 @@
+// Copyright 1998-2019 Epic Games, Inc. All Rights Reserved.
+
+#include "MagicLeapLocationFunctionLibrary.h"
+#include "MagicLeapLocationPlugin.h"
+
+bool UMagicLeapLocationFunctionLibrary::GetLastCoarseLocation(FLocationData& OutLocation)
+{
+	return GET_MAGIC_LEAP_LOCATION_PLUGIN()->GetLastCoarseLocation(OutLocation);
+}
+
+bool UMagicLeapLocationFunctionLibrary::GetLastCoarseLocationAsync(const FCoarseLocationResultDelegate& InResultDelegate)
+{
+	FCoarseLocationResultDelegateMulti ResultDelegate;
+	ResultDelegate.Add(InResultDelegate);
+	return GET_MAGIC_LEAP_LOCATION_PLUGIN()->GetLastCoarseLocationAsync(ResultDelegate);
+}
+
+bool UMagicLeapLocationFunctionLibrary::GetLastCoarseLocationOnSphere(float InRadius, FVector& OutLocation)
+{
+	return GET_MAGIC_LEAP_LOCATION_PLUGIN()->GetLastCoarseLocationOnSphere(InRadius, OutLocation);
+}
+
+bool UMagicLeapLocationFunctionLibrary::GetLastCoarseLocationOnSphereAsync(const FCoarseLocationOnSphereResultDelegate& InResultDelegate, float InRadius)
+{
+	FCoarseLocationOnSphereResultDelegateMulti ResultDelegate;
+	ResultDelegate.Add(InResultDelegate);
+	return GET_MAGIC_LEAP_LOCATION_PLUGIN()->GetLastCoarseLocationOnSphereAsync(ResultDelegate, InRadius);
+}
diff --git a/Engine/Plugins/Lumin/MagicLeapLocation/Source/Private/MagicLeapLocationPlugin.cpp b/Engine/Plugins/Lumin/MagicLeapLocation/Source/Private/MagicLeapLocationPlugin.cpp
new file mode 100644
--- /dev/null
+++ b/Engine/Plugins/Lumin/MagicLeapLocation/Source/Private/MagicLeapLocationPlugin.cpp
@@ -0,0 +1,134 @@
+// Copyright 1998-2019 Epic Games, Inc. All Rights Reserved.
+
+#include "MagicLeapLocationPlugin.h"
+
+#include "Lumin/CAPIShims/LuminAPILocation.h"
+
+DEFINE_LOG_CATEGORY(LogMagicLeapLocation);
+
+using namespace MagicLeap;
+
+FMagicLeapLocationPlugin::FMagicLeapLocationPlugin()
+: PrivilegesManager({ EMagicLeapPrivilege::CoarseLocation })
+, CachedRadius(0.0f)
+{
+}
+
+bool FMagicLeapLocationPlugin::GetLastCoarseLocation(FLocationData& OutLocation)
+{
+#if WITH_MLSDK
+	if (PrivilegesManager.GetPrivilegeStatus(EMagicLeapPrivilege::CoarseLocation, false) == MagicLeap::EPrivilegeState::Granted)
+	{
+		MLLocationData LocationData;
+		MLLocationDataInit(&LocationData);
+		MLResult Result = MLLocationGetLastCoarseLocation(&LocationData);
+		if (Result != MLResult_Ok)
+		{
+			UE_LOG(LogMagicLeapLocation, Error, TEXT("MLLocationGetLastCoarseLocation failed with error '%s'"), UTF8_TO_TCHAR(MLGetResultString(Result)));
+			return false;
+		}
+		else
+		{
+			OutLocation.Latitude = LocationData.latitude;
+			OutLocation.Longitude = LocationData.longitude;
+			if (LocationData.location_mask & MLLocationDataMask_HasPostalCode)
+			{
+				OutLocation.PostalCode = LocationData.postal_code;
+			}
+			return true;
+		}
+	}
+#endif // WITH_MLSDK
+
+	return false;
+}
+
+bool FMagicLeapLocationPlugin::GetLastCoarseLocationAsync(const FCoarseLocationResultDelegateMulti& ResultDelegate)
+{
+	FLocationData LocationData;
+	if (GetLastCoarseLocation(LocationData))
+	{
+		ResultDelegate.Broadcast(LocationData, true);
+		return true;
+	}
+	else
+	{
+		OnCoarseLocationResult = ResultDelegate;
+		PrivilegesManager.SetPrivilegeEventHandler([this](const FRequiredPrivilege& RequiredPrivilege)
+		{
+			OnGetLocationPrivilegeResult(RequiredPrivilege);
+		});
+	}
+
+	return false;
+}
+
+bool FMagicLeapLocationPlugin::GetLastCoarseLocationOnSphere(float InRadius, FVector& OutLocation)
+{
+	FLocationData LocationData;
+	if (!GetLastCoarseLocation(LocationData))
+	{
+		return false;
+	}
+
+	FRotator Rotator;
+	Rotator.Pitch = LocationData.Latitude;
+	Rotator.Yaw = -LocationData.Longitude;
+	Rotator.Roll = 0.0f;
+	OutLocation = FVector(1.0f, 0.0f, 0.0f);
+	OutLocation = Rotator.RotateVector(OutLocation) * InRadius;
+	return true;
+}
+
+bool FMagicLeapLocationPlugin::GetLastCoarseLocationOnSphereAsync(const FCoarseLocationOnSphereResultDelegateMulti& ResultDelegate, float InRadius)
+{
+	if (PrivilegesManager.GetPrivilegeStatus(EMagicLeapPrivilege::CoarseLocation, false) == MagicLeap::EPrivilegeState::Granted)
+	{
+		FVector LocationOnSphere;
+		if (GetLastCoarseLocationOnSphere(InRadius, LocationOnSphere))
+		{
+			ResultDelegate.Broadcast(LocationOnSphere, true);
+		}
+	}
+	else
+	{
+		OnCoarseLocationOnSphereResult = ResultDelegate;
+		CachedRadius = InRadius;
+		PrivilegesManager.SetPrivilegeEventHandler([this](const FRequiredPrivilege& RequiredPrivilege)
+		{
+			OnGetLocationOnSpherePrivilegeResult(RequiredPrivilege);
+		});
+	}
+
+	return true;
+}
+
+void FMagicLeapLocationPlugin::OnGetLocationPrivilegeResult(const FRequiredPrivilege& RequiredPrivilege)
+{
+	FLocationData LocationData;
+	if (RequiredPrivilege.State == EPrivilegeState::Granted)
+	{
+		bool bSuccess = GetLastCoarseLocation(LocationData);
+		OnCoarseLocationResult.Broadcast(LocationData, bSuccess);
+	}
+	else
+	{
+		OnCoarseLocationResult.Broadcast(LocationData, false);
+	}
+}
+
+void FMagicLeapLocationPlugin::OnGetLocationOnSpherePrivilegeResult(const FRequiredPrivilege& RequiredPrivilege)
+{
+	FVector LocationOnSphere;
+	if (RequiredPrivilege.State == EPrivilegeState::Granted)
+	{
+		bool bSuccess = GetLastCoarseLocationOnSphere(CachedRadius, LocationOnSphere);
+		OnCoarseLocationOnSphereResult.Broadcast(LocationOnSphere, bSuccess);
+	}
+	else
+	{
+		OnCoarseLocationOnSphereResult.Broadcast(LocationOnSphere, false);
+	}
+}
+
+IMPLEMENT_MODULE(FMagicLeapLocationPlugin, MagicLeapLocation);
diff --git a/Engine/Plugins/Lumin/MagicLeapLocation/Source/Private/MagicLeapLocationPlugin.h b/Engine/Plugins/Lumin/MagicLeapLocation/Source/Private/MagicLeapLocationPlugin.h
new file mode 100644
--- /dev/null
+++ b/Engine/Plugins/Lumin/MagicLeapLocation/Source/Private/MagicLeapLocationPlugin.h
@@ -0,0 +1,32 @@
+// Copyright 1998-2019 Epic Games, Inc. All Rights Reserved.
+
+#pragma once
+
+#include "CoreMinimal.h"
+#include "IMagicLeapLocationPlugin.h"
+#include "MagicLeapLocationTypes.h"
+#include "AppEventHandler.h"
+#include "MagicLeapPluginUtil.h"
+
+DECLARE_LOG_CATEGORY_EXTERN(LogMagicLeapLocation, Verbose, All);
+
+class FMagicLeapLocationPlugin : public IMagicLeapLocationPlugin
+{
+public:
+	FMagicLeapLocationPlugin();
+
+	bool GetLastCoarseLocation(FLocationData& OutLocation);
+	bool GetLastCoarseLocationAsync(const FCoarseLocationResultDelegateMulti& ResultDelegate);
+	bool GetLastCoarseLocationOnSphere(float InRadius, FVector& OutLocation);
+	bool GetLastCoarseLocationOnSphereAsync(const FCoarseLocationOnSphereResultDelegateMulti& ResultDelegate, float InRadius);
+private:
+	MagicLeap::IAppEventHandler PrivilegesManager;
+	float CachedRadius;
+	FCoarseLocationResultDelegateMulti OnCoarseLocationResult;
+	FCoarseLocationOnSphereResultDelegateMulti OnCoarseLocationOnSphereResult;
+
+	void OnGetLocationPrivilegeResult(const MagicLeap::FRequiredPrivilege& RequiredPrivilege);
+	void OnGetLocationOnSpherePrivilegeResult(const MagicLeap::FRequiredPrivilege& RequiredPrivilege);
+};
+
+#define GET_MAGIC_LEAP_LOCATION_PLUGIN() static_cast<FMagicLeapLocationPlugin*>(&IMagicLeapLocationPlugin::Get())
diff --git a/Engine/Plugins/Lumin/MagicLeapLocation/Source/Public/IMagicLeapLocationPlugin.h b/Engine/Plugins/Lumin/MagicLeapLocation/Source/Public/IMagicLeapLocationPlugin.h
new file mode 100644
--- /dev/null
+++ b/Engine/Plugins/Lumin/MagicLeapLocation/Source/Public/IMagicLeapLocationPlugin.h
@@ -0,0 +1,35 @@
+// Copyright 1998-2019 Epic Games, Inc. All Rights Reserved.
+
+#pragma once
+
+#include "Modules/ModuleManager.h"
+
+/**
+ * The public interface to this module.  In most cases, this interface is only public to sibling modules
+ * within this plugin.
+ */
+class IMagicLeapLocationPlugin : public IModuleInterface
+{
+public:
+
+	/**
+	 * Singleton-like access to this module's interface.  This is just for convenience!
+	 * Beware of calling this during the shutdown phase, though.  Your module might have been unloaded already.
+	 *
+	 * @return Returns singleton instance, loading the module on demand if needed
+	 */
+	static inline IMagicLeapLocationPlugin& Get()
+	{
+		return FModuleManager::LoadModuleChecked<IMagicLeapLocationPlugin>("MagicLeapLocation");
+	}
+
+	/**
+	 * Checks to see if this module is loaded and ready.  It is only valid to call Get() if IsAvailable() returns true.
+	 *
+	 * @return True if the module is loaded and ready to use
+	 */
+	static inline bool IsAvailable()
+	{
+		return FModuleManager::Get().IsModuleLoaded("MagicLeapLocation");
+	}
+};
diff --git a/Engine/Plugins/Lumin/MagicLeapLocation/Source/Public/MagicLeapLocationComponent.h b/Engine/Plugins/Lumin/MagicLeapLocation/Source/Public/MagicLeapLocationComponent.h
new file mode 100644
--- /dev/null
+++ b/Engine/Plugins/Lumin/MagicLeapLocation/Source/Public/MagicLeapLocationComponent.h
@@ -0,0 +1,61 @@
+// Copyright 1998-2019 Epic Games, Inc. All Rights Reserved.
+
+#pragma once
+
+#include "Engine/Engine.h"
+#include "Components/ActorComponent.h"
+#include "MagicLeapLocationTypes.h"
+#include "MagicLeapLocationComponent.generated.h"
+
+/**
+	Component that provides access to the Location API functionality.
+*/
+UCLASS(ClassGroup = MagicLeap, BlueprintType, Blueprintable, EditInlineNew, meta = (BlueprintSpawnableComponent))
+class MAGICLEAPLOCATION_API UMagicLeapLocationComponent : public UActorComponent
+{
+	GENERATED_BODY()
+
+public:
+	UMagicLeapLocationComponent();
+	virtual ~UMagicLeapLocationComponent();
+
+	/**
+		Attempts to retrieve the latitude, longitude and postcode of the device.
+		@param OutLocation If successful this will contain the latitude, longitude and postcode of the device.
+		@return True if the location data is valid, false otherwise.
+	*/
+	UFUNCTION(BlueprintCallable, Category = "Location | MagicLeap")
+	bool GetLastCoarseLocation(FLocationData& OutLocation);
+
+	/**
+		Attempts to retrieve the latitude, longitude and postcode of the device asynchronously.
+		@return True if the location is immediately resolved, false otherwise.
+	*/
+	UFUNCTION(BlueprintCallable, Category = "Location | MagicLeap")
+	bool GetLastCoarseLocationAsync();
+
+	/**
+		Attempts to retrieve a point on a sphere representing the location of the device.
+		@param InRadius The radius of the sphere that the location will be projected onto.
+		@param OutLocation If successful this will be a valid point on a sphere representing the location of the device.
+		@return True if the location is valid, false otherwise.
+	*/
+	UFUNCTION(BlueprintCallable, Category = "Location | MagicLeap")
+	bool GetLastCoarseLocationOnSphere(float InRadius, FVector& OutLocation);
+
+	/**
+		Attempts to retrieve a point on a sphere representing the location of the device asynchronously.
+		@param InRadius The radius of the sphere that the location will be projected onto.
+		@return True if the location is immediately resolved, false otherwise.
+	*/
+	UFUNCTION(BlueprintCallable, Category = "Location | MagicLeap")
+	bool GetLastCoarseLocationOnSphereAsync(float InRadius);
+
+private:
+	// Delegate instances
+	UPROPERTY(BlueprintAssignable, Category = "Location | MagicLeap", meta = (AllowPrivateAccess = true))
+	FCoarseLocationResultDelegateMulti OnGotCoarseLocation;
+
+	UPROPERTY(BlueprintAssignable, Category = "Location | MagicLeap", meta = (AllowPrivateAccess = true))
+	FCoarseLocationOnSphereResultDelegateMulti OnGotCoarseLocationOnSphere;
+};
diff --git a/Engine/Plugins/Lumin/MagicLeapLocation/Source/Public/MagicLeapLocationFunctionLibrary.h b/Engine/Plugins/Lumin/MagicLeapLocation/Source/Public/MagicLeapLocationFunctionLibrary.h
new file mode 100644
--- /dev/null
+++ b/Engine/Plugins/Lumin/MagicLeapLocation/Source/Public/MagicLeapLocationFunctionLibrary.h
@@ -0,0 +1,47 @@
+// Copyright 1998-2019 Epic Games, Inc. All Rights Reserved.
+#pragma once
+
+#include "Engine/Engine.h"
+#include "Kismet/BlueprintFunctionLibrary.h"
+#include "MagicLeapLocationTypes.h"
+#include "MagicLeapLocationFunctionLibrary.generated.h"
+
+UCLASS(ClassGroup = MagicLeap)
+class MAGICLEAPLOCATION_API UMagicLeapLocationFunctionLibrary : public UBlueprintFunctionLibrary
+{
+	GENERATED_BODY()
+
+public:
+	/**
+		Attempts to retrieve the latitude, longitude and postcode of the device.
+		@param OutLocation If successful this will contain the latitude, longitude and postcode of the device.
+		@return True if the location data is valid, false otherwise.
+	*/
+	UFUNCTION(BlueprintCallable, Category = "Location Function Library | MagicLeap")
+	static bool GetLastCoarseLocation(FLocationData& OutLocation);
+
+	/**
+		Attempts to retrieve the latitude, longitude and postcode of the device asynchronously.
+		@param The delegate to notify once the privilege has been granted.
+		@return True if the location is immediately resolved, false otherwise.
+	*/
+	UFUNCTION(BlueprintCallable, Category = "Location | MagicLeap")
+	bool GetLastCoarseLocationAsync(const FCoarseLocationResultDelegate& InResultDelegate);
+
+	/**
+		Attempts to retrieve a point on a sphere representing the location of the device.
+		@param OutLocation If successful this will be a valid point on a sphere representing the location of the device.
+		@return True if the location is valid, false otherwise.
+	*/
+	UFUNCTION(BlueprintCallable, Category = "Location Function Library | MagicLeap")
+	static bool GetLastCoarseLocationOnSphere(float InRadius, FVector& OutLocation);
+
+	/**
+		Attempts to retrieve a point on a sphere representing the location of the device asynchronously.
+		@param The delegate to notify once the privilege has been granted.
+		@param InRadius The radius of the sphere that the location will be projected onto.
+		@return True if the location is immediately resolved, false otherwise.
+	*/
+	UFUNCTION(BlueprintCallable, Category = "Location | MagicLeap")
+	bool GetLastCoarseLocationOnSphereAsync(const FCoarseLocationOnSphereResultDelegate& InResultDelegate, float InRadius);
+};
diff --git a/Engine/Plugins/Lumin/MagicLeapLocation/Source/Public/MagicLeapLocationTypes.h b/Engine/Plugins/Lumin/MagicLeapLocation/Source/Public/MagicLeapLocationTypes.h
new file mode 100644
--- /dev/null
+++ b/Engine/Plugins/Lumin/MagicLeapLocation/Source/Public/MagicLeapLocationTypes.h
@@ -0,0 +1,45 @@
+// Copyright 1998-2019 Epic Games, Inc. All Rights Reserved.
+#pragma once
+
+#include "Engine/Engine.h"
+#include "MagicLeapLocationTypes.generated.h"
+
+UENUM(BlueprintType)
+enum class EMagicLeapLocationResult : uint8
+{
+	/** Unknown location error. */
+	Unknown,
+	/** No connection to server. */
+	NetworkConnection,
+	/** No location data received. */
+	NoLocation,
+	/** Location provider is not found. */
+	ProviderNotFound,
+};
+
+/** Location request result. */
+USTRUCT(BlueprintType)
+struct FLocationData
+{
+	GENERATED_USTRUCT_BODY()
+
+	/** Location latitude. */
+	UPROPERTY(EditAnywhere, BlueprintReadWrite, Category = "Location|MagicLeap")
+	float Latitude;
+
+	/** Location longitude. */
+	UPROPERTY(EditAnywhere, BlueprintReadWrite, Category = "Location|MagicLeap")
+	float Longitude;
+
+	/** Location longitude. */
+	UPROPERTY(EditAnywhere, BlueprintReadWrite, Category = "Location|MagicLeap")
+	FString PostalCode;
+};
+
+/** Delegate used to convey the result of a coarse location query. */
+DECLARE_DYNAMIC_DELEGATE_TwoParams(FCoarseLocationResultDelegate, const FLocationData&, LocationData, bool, bSuccess);
+DECLARE_DYNAMIC_MULTICAST_DELEGATE_TwoParams(FCoarseLocationResultDelegateMulti, const FLocationData&, LocationData, bool, bSuccess);
+
+/** Delegate used to convey the result of a coarse location on sphere query. */
+DECLARE_DYNAMIC_DELEGATE_TwoParams(FCoarseLocationOnSphereResultDelegate, const FVector&, LocationOnSphere, bool, bSuccess);
+DECLARE_DYNAMIC_MULTICAST_DELEGATE_TwoParams(FCoarseLocationOnSphereResultDelegateMulti, const FVector&, LocationOnSphere, bool, bSuccess);
diff --git a/Engine/Plugins/Lumin/MagicLeapMedia/MagicLeapMedia.uplugin b/Engine/Plugins/Lumin/MagicLeapMedia/MagicLeapMedia.uplugin
--- a/Engine/Plugins/Lumin/MagicLeapMedia/MagicLeapMedia.uplugin
+++ b/Engine/Plugins/Lumin/MagicLeapMedia/MagicLeapMedia.uplugin
@@ -36,27 +36,25 @@
 			"Type" : "Editor",
 			"LoadingPhase": "PostEngineInit"
 		},
-        {
-            "Name": "MagicLeapMediaFactory",
-            "Type": "RuntimeNoCommandlet",
-            "LoadingPhase": "PostEngineInit",
-            "WhitelistPlatforms": [ "Lumin", "Mac", "Win64" ]
-        },
-        {
+	        {
+	            "Name": "MagicLeapMediaFactory",
+	            "Type": "RuntimeNoCommandlet",
+	            "LoadingPhase": "PostEngineInit",
+	            "WhitelistPlatforms": [ "Lumin", "Mac", "Win64" ]
+	        },
+	        {
 			"Name" : "MagicLeapMediaCodecFactory",
 			"Type" : "RuntimeNoCommandlet",
 			"LoadingPhase" : "PostEngineInit",
 			"WhitelistPlatforms" : [ "Lumin", "Mac", "Win64" ]
-		},
-        {
-			"Name" : "MagicLeapCameraPreviewFactory",
-			"Type" : "RuntimeNoCommandlet",
-			"LoadingPhase" : "PostEngineInit",
-			"WhitelistPlatforms" : [ "Lumin", "Mac", "Win64" ]
 		}
 	],
 	"Plugins":
 	[
+		{
+			"Name": "MLSDK",
+			"Enabled": true
+		},
 		{
 			"Name": "MagicLeap",
 			"Enabled": true
diff --git a/Engine/Plugins/Lumin/MagicLeapMedia/Source/MagicLeapMedia/MagicLeapMedia.Build.cs b/Engine/Plugins/Lumin/MagicLeapMedia/Source/MagicLeapMedia/MagicLeapMedia.Build.cs
--- a/Engine/Plugins/Lumin/MagicLeapMedia/Source/MagicLeapMedia/MagicLeapMedia.Build.cs
+++ b/Engine/Plugins/Lumin/MagicLeapMedia/Source/MagicLeapMedia/MagicLeapMedia.Build.cs
@@ -9,7 +9,6 @@ namespace UnrealBuildTool.Rules
 		public MagicLeapMedia(ReadOnlyTargetRules Target) : base(Target)
 		{
             PublicIncludePaths.Add(Path.Combine(ModuleDirectory, "Public"));
-			string EngineDir = Path.GetFullPath(Target.RelativeEnginePath);
 
 			DynamicallyLoadedModuleNames.AddRange(
 				new string[] {
@@ -22,11 +21,11 @@ namespace UnrealBuildTool.Rules
 					"Engine",
 					"RenderCore",
 					"MediaUtils",
-					"OpenGLDrv",
 					"LuminRuntimeSettings",
 					"MLSDK",
                     "MagicLeap",
                     "MagicLeapHelperVulkan",
+					"MagicLeapHelperOpenGL"
             });
 
 			PrivateIncludePathModuleNames.AddRange(
@@ -36,11 +35,8 @@ namespace UnrealBuildTool.Rules
 
 			PrivateIncludePaths.AddRange(
 				new string[] {
-                    "MagicLeapMedia/Private",
-					Path.Combine(EngineDir, "Source/Runtime/OpenGLDrv/Private")
+                    "MagicLeapMedia/Private"
             });
-
-			AddEngineThirdPartyPrivateStaticDependencies(Target, "OpenGL");
 		}
 	}
 }
diff --git a/Engine/Plugins/Lumin/MagicLeapMedia/Source/MagicLeapMedia/Private/MagicLeapMediaModule.cpp b/Engine/Plugins/Lumin/MagicLeapMedia/Source/MagicLeapMedia/Private/MagicLeapMediaModule.cpp
--- a/Engine/Plugins/Lumin/MagicLeapMedia/Source/MagicLeapMedia/Private/MagicLeapMediaModule.cpp
+++ b/Engine/Plugins/Lumin/MagicLeapMedia/Source/MagicLeapMedia/Private/MagicLeapMediaModule.cpp
@@ -2,7 +2,6 @@
 
 #include "CoreMinimal.h"
 #include "MagicLeapMediaPlayer.h"
-#include "MagicLeapCameraPreviewPlayer.h"
 #include "IMagicLeapMediaModule.h"
 #include "Modules/ModuleManager.h"
 
@@ -17,11 +16,6 @@ public:
 		return MakeShareable(new FMagicLeapMediaPlayer(EventSink));
 	}
 
-	virtual TSharedPtr<IMediaPlayer, ESPMode::ThreadSafe> CreateCameraPreviewPlayer(IMediaEventSink& EventSink)
-	{
-		return MakeShareable(new FMagicLeapCameraPreviewPlayer(EventSink));
-	}
-
 public:
 	/** IModuleInterface interface */
 	virtual void StartupModule() override { }
diff --git a/Engine/Plugins/Lumin/MagicLeapMedia/Source/MagicLeapMedia/Private/MagicLeapMediaOverlaySample.h b/Engine/Plugins/Lumin/MagicLeapMedia/Source/MagicLeapMedia/Private/MagicLeapMediaOverlaySample.h
new file mode 100644
--- /dev/null
+++ b/Engine/Plugins/Lumin/MagicLeapMedia/Source/MagicLeapMedia/Private/MagicLeapMediaOverlaySample.h
@@ -0,0 +1,60 @@
+// Copyright 1998-2019 Epic Games, Inc. All Rights Reserved.
+
+#pragma once
+
+#include "CoreTypes.h"
+#include "IMediaOverlaySample.h"
+#include "Misc/Timespan.h"
+
+class FMagicLeapMediaOverlaySample : public IMediaOverlaySample
+{
+public:
+	FMagicLeapMediaOverlaySample()
+		: Duration(FTimespan::Zero())
+		, Time(FTimespan::Zero())
+	{ }
+
+	void Intitialize(char* Data, const EMediaOverlaySampleType InSampleType, const FTimespan InDuration, const FTimespan InTime)
+	{
+		check(Data != nullptr);
+		OverlaySampleType = InSampleType;
+		Text = FText::FromString(FString(ANSI_TO_TCHAR(Data)));
+		Duration = InDuration;
+		Time = InTime;
+	}
+
+	virtual ~FMagicLeapMediaOverlaySample() {};
+
+	/** IMediaOverlaySample interface **/
+	virtual FTimespan GetDuration() const override
+	{
+		return Duration;
+	}
+
+	virtual TOptional<FVector2D> GetPosition() const override
+	{
+		return TOptional<FVector2D>();
+	}
+
+	virtual FText GetText() const override
+	{
+		return Text;
+	}
+
+	virtual FTimespan GetTime() const override
+	{
+		return Time;
+	}
+
+	virtual EMediaOverlaySampleType GetType() const override
+	{
+		return OverlaySampleType;
+	}
+
+private:
+	EMediaOverlaySampleType OverlaySampleType;
+	FText Text;
+	FTimespan Duration;
+	FTimespan Time;
+};
+
diff --git a/Engine/Plugins/Lumin/MagicLeapMedia/Source/MagicLeapMedia/Private/MagicLeapMediaPlayer.cpp b/Engine/Plugins/Lumin/MagicLeapMedia/Source/MagicLeapMedia/Private/MagicLeapMediaPlayer.cpp
--- a/Engine/Plugins/Lumin/MagicLeapMedia/Source/MagicLeapMedia/Private/MagicLeapMediaPlayer.cpp
+++ b/Engine/Plugins/Lumin/MagicLeapMedia/Source/MagicLeapMedia/Private/MagicLeapMediaPlayer.cpp
@@ -4,36 +4,46 @@
 #include "IMagicLeapMediaModule.h"
 #include "Misc/Paths.h"
 #include "Lumin/LuminPlatformFile.h"
+#include "Lumin/CAPIShims/LuminAPILogging.h"
+#include "Lumin/CAPIShims/LuminAPIMediaError.h"
+#include "Lumin/CAPIShims/LuminAPIMediaFormat.h"
+#include "Lumin/CAPIShims/LuminAPISharedFile.h"
+#include "Lumin/CAPIShims/LuminAPIFileInfo.h"
 #include "IMediaEventSink.h"
+#include "IMediaOptions.h"
+#include "MagicLeapMediaOverlaySample.h"
 #include "MediaSamples.h"
 #include "Misc/CoreDelegates.h"
 #include "UObject/Class.h"
 #include "UObject/UObjectGlobals.h"
 #include "ExternalTexture.h"
-#include "MagicLeapMediaTextureSample.h"
-#include "ExternalOESTextureRenderer.h"
-#include "Lumin/LuminEGL.h"
-#include "MediaWorker.h"
-#include "CameraCaptureComponent.h"
+#include "ExternalTextureRendererGL.h"
 #include "MagicLeapHelperVulkan.h"
+#include "HAL/Platform.h"
 
-#ifndef EGL_EGLEXT_PROTOTYPES
-#define EGL_EGLEXT_PROTOTYPES
-#endif // !EGL_EGLEXT_PROTOTYPES
+#define LOCTEXT_NAMESPACE "FMagicLeapMediaModule"
 
-#include <EGL/egl.h>
-#include <EGL/eglext.h>
+static void *STDCALL _ThreadProc(void *PlayerHandle)
+{
+	FPlatformMisc::LowLevelOutputDebugString(TEXT("In the detatched thread to destroy media player"));
 
-#ifndef GL_GLEXT_PROTOTYPES
-#define GL_GLEXT_PROTOTYPES
-#endif // !GL_GLEXT_PROTOTYPES
+	MLHandle MediaPlayerHandle = reinterpret_cast<MLHandle>(PlayerHandle);
 
-#include <GLES2/gl2.h>
-#include <GLES2/gl2ext.h>
+	MLResult Result = MLMediaPlayerReset(MediaPlayerHandle);
+	if (Result != MLResult_Ok)
+	{
+		FPlatformMisc::LowLevelOutputDebugStringf(TEXT("MLMediaPlayerReset failed with error %s"), MLMediaResultGetString(Result));
+	}
 
-#include "ml_media_player.h"
+	Result = MLMediaPlayerDestroy(MediaPlayerHandle);
+	if (Result != MLResult_Ok)
+	{
+		FPlatformMisc::LowLevelOutputDebugStringf(TEXT("MLMediaPlayerDestroy failed with error %s"), MLMediaResultGetString(Result));
+	}
 
-#define LOCTEXT_NAMESPACE "FMagicLeapMediaModule"
+	pthread_exit(nullptr);
+	return nullptr;
+}
 
 struct FMagicLeapVideoTextureData
 {
@@ -59,90 +69,18 @@ struct FMagicLeapVideoTextureDataGL : public FMagicLeapVideoTextureData
 public:
 
 	FMagicLeapVideoTextureDataGL()
-	: Image(EGL_NO_IMAGE_KHR)
-	, ExternalRenderer(GSupportsImageExternal ? nullptr : new FExternalOESTextureRenderer(false))
-	, Display(EGL_NO_DISPLAY)
-	, Context(EGL_NO_CONTEXT)
-	, SavedDisplay(EGL_NO_DISPLAY)
-	, SavedContext(EGL_NO_CONTEXT)
-	, bContextCreated(false)
 	{}
 
-	~FMagicLeapVideoTextureDataGL()
-	{
-		delete ExternalRenderer;
-		ExternalRenderer = nullptr;
-		PreviousNativeBuffer = ML_INVALID_HANDLE;
-		eglDestroyContext(Display, Context);
-		Display = EGL_NO_DISPLAY;
-		Context = EGL_NO_CONTEXT;   
-	}
-
-	bool InitContext()
-	{
-#if !PLATFORM_LUMINGL4
-		if (Context == EGL_NO_CONTEXT)
-		{
-			Display = LuminEGL::GetInstance()->GetDisplay();
-			EGLContext SharedContext = LuminEGL::GetInstance()->GetCurrentContext();
-			Context = SharedContext; // LuminEGL::GetInstance()->CreateContext(SharedContext);
-		}
-
-		return (Context != EGL_NO_CONTEXT);
-#else
-		return false;
-#endif
-	}
-
-	void SaveContext()
-	{
-#if !PLATFORM_LUMINGL4
-		SavedDisplay = LuminEGL::GetInstance()->GetDisplay();
-		SavedContext = LuminEGL::GetInstance()->GetCurrentContext();
-#endif
-	}
-
-	void MakeCurrent()
-	{
-#if !PLATFORM_LUMINGL4
-		return;	// skip for now
-		EGLBoolean bResult = eglMakeCurrent(Display, EGL_NO_SURFACE, EGL_NO_SURFACE, Context);
-		if (bResult == EGL_FALSE)
-		{
-			UE_LOG(LogMagicLeapMedia, Error, TEXT("Error setting media player context."));
-		}
-#endif
-	}
-
-	void RestoreContext()
-	{
-#if !PLATFORM_LUMINGL4
-		return;	// skip for now
-		EGLBoolean bResult = eglMakeCurrent(SavedDisplay, EGL_NO_SURFACE, EGL_NO_SURFACE, SavedContext);
-		if (bResult == EGL_FALSE)
-		{
-			UE_LOG(LogMagicLeapMedia, Error, TEXT("Error setting unreal context."));
-		}
-#endif
-	}
-
-	EGLImageKHR Image;
-	FExternalOESTextureRenderer* ExternalRenderer;
-	EGLDisplay Display;
-	EGLContext Context;
-	EGLDisplay SavedDisplay;
-	EGLContext SavedContext;
-	bool bContextCreated;
+	FExternalTextureRendererGL TextureRenderer;
 };
 
 FMagicLeapMediaPlayer::FMagicLeapMediaPlayer(IMediaEventSink& InEventSink)
 : MediaPlayerHandle(ML_INVALID_HANDLE)
+, CaptionParserHandle(ML_INVALID_HANDLE)
 , bMediaPrepared(false)
 , CurrentState(EMediaState::Closed)
 , EventSink(InEventSink)
 , Samples(MakeShared<FMediaSamples, ESPMode::ThreadSafe>())
-, VideoSamplePool(new FMagicLeapMediaTextureSamplePool())
-, MediaWorker(nullptr)
 , bWasMediaPlayingBeforeAppPause(false)
 , bPlaybackCompleted(false)
 , CurrentPlaybackTime(FTimespan::Zero())
@@ -154,113 +92,12 @@ FMagicLeapMediaPlayer::FMagicLeapMediaPlayer(IMediaEventSink& InEventSink)
 	else
 	{
 		TextureData = MakeShared<FMagicLeapVideoTextureDataGL, ESPMode::ThreadSafe>();
-	}
-	
-	MLResult Result = MLMediaPlayerCreate(&MediaPlayerHandle);
-	UE_CLOG(Result != MLResult_Ok, LogMagicLeapMedia, Error, TEXT("MLMediaPlayerCreate failed with error %s."), UTF8_TO_TCHAR(MLMediaResultGetString(Result)));
-	CurrentState = (Samples.IsValid() && Result == MLResult_Ok) ? EMediaState::Closed : EMediaState::Error;
-	
-	if (!GSupportsImageExternal)
-	{
-		MediaWorker = new FMediaWorker(MediaPlayerHandle, CriticalSection);
-	}
+	}	
 }
 
 FMagicLeapMediaPlayer::~FMagicLeapMediaPlayer()
 {
 	Close();
-
-	delete MediaWorker;
-	MediaWorker = nullptr;
-
-	if (MLHandleIsValid(MediaPlayerHandle))
-	{
-		if (GSupportsImageExternal)
-		{
-			if (FLuminPlatformMisc::ShouldUseVulkan())
-			{
-				// Unregister the external texture on render thread
-				struct FReleaseVideoResourcesParams
-				{
-					FMagicLeapMediaPlayer* MediaPlayer;
-					TSharedPtr<FMagicLeapVideoTextureDataVK, ESPMode::ThreadSafe> TextureData;
-					FGuid PlayerGuid;
-					MLHandle MediaPlayerHandle;
-				};
-
-				FReleaseVideoResourcesParams Params = { this, StaticCastSharedPtr<FMagicLeapVideoTextureDataVK>(TextureData), PlayerGuid, MediaPlayerHandle };
-
-				ENQUEUE_RENDER_COMMAND(MagicLeapMediaPlayerDestroy)(
-					[Params](FRHICommandListImmediate& RHICmdList)
-					{
-						FExternalTextureRegistry::Get().UnregisterExternalTexture(Params.PlayerGuid);
-
-						if (Params.TextureData->PreviousNativeBuffer != 0 && MLHandleIsValid(Params.TextureData->PreviousNativeBuffer))
-						{
-							Params.MediaPlayer->RenderThreadReleaseNativeBuffer(Params.MediaPlayerHandle, Params.TextureData->PreviousNativeBuffer);
-						}
-
-						if (MLHandleIsValid(Params.MediaPlayerHandle))
-						{
-							MLResult Result = MLMediaPlayerDestroy(Params.MediaPlayerHandle);
-							UE_CLOG(Result != MLResult_Ok, LogMagicLeapMedia, Error, TEXT("MLMediaPlayerDestroy failed with error %s."), UTF8_TO_TCHAR(MLMediaResultGetString(Result)));
-						}
-					});
-			}
-			else
-			{
-				// Unregister the external texture on render thread
-				struct FReleaseVideoResourcesParams
-				{
-					FMagicLeapMediaPlayer* MediaPlayer;
-					// Can we make this a TWeakPtr? We need to ensure that FMagicLeapMediaPlayer::TextureData is not destroyed before
-					// we unregister the external texture and releae the VideoTexture.
-					TSharedPtr<FMagicLeapVideoTextureDataGL, ESPMode::ThreadSafe> TextureData;
-					FGuid PlayerGuid;
-					MLHandle MediaPlayerHandle;
-				};
-
-				FReleaseVideoResourcesParams Params = { this, StaticCastSharedPtr<FMagicLeapVideoTextureDataGL>(TextureData), PlayerGuid, MediaPlayerHandle };
-
-				ENQUEUE_RENDER_COMMAND(MagicLeapMediaPlayerDestroy)(
-					[Params](FRHICommandListImmediate& RHICmdList)
-					{
-						FExternalTextureRegistry::Get().UnregisterExternalTexture(Params.PlayerGuid);
-						// @todo: this causes a crash
-						//Params.TextureData->VideoTexture->Release();
-						Params.TextureData->SaveContext();
-						Params.TextureData->MakeCurrent();
-
-						if (Params.TextureData->Image != EGL_NO_IMAGE_KHR)
-						{
-							eglDestroyImageKHR(eglGetCurrentDisplay(), Params.TextureData->Image);
-							Params.TextureData->Image = EGL_NO_IMAGE_KHR;
-						}
-
-						Params.TextureData->RestoreContext();
-						if (Params.TextureData->PreviousNativeBuffer != 0 && MLHandleIsValid(Params.TextureData->PreviousNativeBuffer))
-						{
-							Params.MediaPlayer->RenderThreadReleaseNativeBuffer(Params.MediaPlayerHandle, Params.TextureData->PreviousNativeBuffer);
-						}
-
-						MLResult Result = MLMediaPlayerDestroy(Params.MediaPlayerHandle);
-						UE_CLOG(Result != MLResult_Ok, LogMagicLeapMedia, Error, TEXT("MLMediaPlayerDestroy failed with error %s."), UTF8_TO_TCHAR(MLMediaResultGetString(Result)));
-					});
-			}
-
-			FlushRenderingCommands();
-		}
-		else
-		{
-			MLResult Result = MLMediaPlayerDestroy(MediaPlayerHandle);
-			UE_CLOG(Result != MLResult_Ok, LogMagicLeapMedia, Error, TEXT("MLMediaPlayerDestroy failed with error %s."), UTF8_TO_TCHAR(MLMediaResultGetString(Result)));
-		}
-
-		MediaPlayerHandle = ML_INVALID_HANDLE;
-	}
-
-	delete VideoSamplePool;
-	VideoSamplePool = nullptr;
 }
 
 void FMagicLeapMediaPlayer::Close()
@@ -291,15 +128,90 @@ void FMagicLeapMediaPlayer::Close()
 	MLResult StopResult = MLMediaPlayerStop(MediaPlayerHandle);
 	UE_CLOG(StopResult != MLResult_Ok, LogMagicLeapMedia, Error, TEXT("MLMediaPlayerStop failed with error %s"), UTF8_TO_TCHAR(MLMediaResultGetString(StopResult)));
 
-	MLResult ResetResult = MLMediaPlayerReset(MediaPlayerHandle);
-	UE_CLOG(ResetResult != MLResult_Ok, LogMagicLeapMedia, Error, TEXT("MLMediaPlayerReset failed with error %s"), UTF8_TO_TCHAR(MLMediaResultGetString(ResetResult)));
+	// Moved here from the constructor temporarily, until the freeze in MLMediaPlayerReset()/Destroy() is fixed.
+	if (GSupportsImageExternal)
+	{
+		if (FLuminPlatformMisc::ShouldUseVulkan())
+		{
+			// Unregister the external texture on render thread
+			struct FReleaseVideoResourcesParams
+			{
+				FMagicLeapMediaPlayer* MediaPlayer;
+				TSharedPtr<FMagicLeapVideoTextureDataVK, ESPMode::ThreadSafe> TextureData;
+				FGuid PlayerGuid;
+				MLHandle MediaPlayerHandle;
+			};
+
+			FReleaseVideoResourcesParams Params = { this, StaticCastSharedPtr<FMagicLeapVideoTextureDataVK>(TextureData), PlayerGuid, MediaPlayerHandle };
+
+			ENQUEUE_RENDER_COMMAND(MagicLeapMediaPlayerDestroy)(
+				[Params](FRHICommandListImmediate& RHICmdList)
+				{
+					FExternalTextureRegistry::Get().UnregisterExternalTexture(Params.PlayerGuid);
+
+					if (Params.TextureData->PreviousNativeBuffer != 0 && MLHandleIsValid(Params.TextureData->PreviousNativeBuffer))
+					{
+						Params.MediaPlayer->RenderThreadReleaseNativeBuffer(Params.MediaPlayerHandle, Params.TextureData->PreviousNativeBuffer);
+					}
+
+					Params.TextureData->VideoTexturePool.Empty();
+
+					Params.MediaPlayer->TriggerResetAndDestroy();
+				});
+		}
+		else
+		{
+			// Unregister the external texture on render thread
+			struct FReleaseVideoResourcesParams
+			{
+				FMagicLeapMediaPlayer* MediaPlayer;
+				// Can we make this a TWeakPtr? We need to ensure that FMagicLeapMediaPlayer::TextureData is not destroyed before
+				// we unregister the external texture and releae the VideoTexture.
+				TSharedPtr<FMagicLeapVideoTextureDataGL, ESPMode::ThreadSafe> TextureData;
+				FGuid PlayerGuid;
+				MLHandle MediaPlayerHandle;
+			};
+
+			FReleaseVideoResourcesParams Params = { this, StaticCastSharedPtr<FMagicLeapVideoTextureDataGL>(TextureData), PlayerGuid, MediaPlayerHandle };
+
+			ENQUEUE_RENDER_COMMAND(MagicLeapMediaPlayerDestroy)(
+				[Params](FRHICommandListImmediate& RHICmdList)
+				{
+					FExternalTextureRegistry::Get().UnregisterExternalTexture(Params.PlayerGuid);
+					// @todo: this causes a crash
+					//Params.TextureData->VideoTexture->Release();
+
+					Params.TextureData->TextureRenderer.DestroyImageKHR();
+
+					if (Params.TextureData->PreviousNativeBuffer != 0 && MLHandleIsValid(Params.TextureData->PreviousNativeBuffer))
+					{
+						Params.MediaPlayer->RenderThreadReleaseNativeBuffer(Params.MediaPlayerHandle, Params.TextureData->PreviousNativeBuffer);
+					}
+
+					Params.MediaPlayer->TriggerResetAndDestroy();
+				});
+		}
+
+		FlushRenderingCommands();
+	}
+	else
+	{
+		TriggerResetAndDestroy();
+	}
+
+	if (MLHandleIsValid(CaptionParserHandle))
+	{
+		MLResult Result = MLMediaCCParserDestroy(CaptionParserHandle);
+		UE_CLOG(Result != MLResult_Ok, LogMagicLeapMedia, Error, TEXT("MLMediaCCParserDestroy failed with error %s."), UTF8_TO_TCHAR(MLGetResultString(Result)));
+		CaptionParserHandle = ML_INVALID_HANDLE;
+	}
 
+	MediaPlayerHandle = ML_INVALID_HANDLE;
 	CurrentState = EMediaState::Closed;
 
 	bMediaPrepared = false;
 	Info.Empty();
 	MediaUrl = FString();
-	VideoSamplePool->Reset();
 	TrackInfo.Empty();
 	SelectedTrack.Empty();
 
@@ -356,6 +268,11 @@ IMediaView& FMagicLeapMediaPlayer::GetView()
 
 bool FMagicLeapMediaPlayer::Open(const FString& Url, const IMediaOptions* Options)
 {
+	// Moved here from the constructor temporarily, until the freeze in MLMediaPlayerReset()/Destroy() is fixed.
+	MLResult Result = MLMediaPlayerCreate(&MediaPlayerHandle);
+	UE_CLOG(Result != MLResult_Ok, LogMagicLeapMedia, Error, TEXT("MLMediaPlayerCreate failed with error %s."), UTF8_TO_TCHAR(MLMediaResultGetString(Result)));
+	CurrentState = (Samples.IsValid() && Result == MLResult_Ok) ? EMediaState::Closed : EMediaState::Error;
+
 	if (CurrentState == EMediaState::Error)
 	{
 		EventSink.ReceiveMediaEvent(EMediaEvent::MediaOpenFailed);
@@ -373,6 +290,7 @@ bool FMagicLeapMediaPlayer::Open(const FString& Url, const IMediaOptions* Option
 	MediaUrl = Url;
 
 	const FString localFileSchema = "file://";
+	const FString sharedFileSchema = "mlshared://";
 
 	// open the media
 	if (Url.StartsWith(localFileSchema))
@@ -391,7 +309,7 @@ bool FMagicLeapMediaPlayer::Open(const FString& Url, const IMediaOptions* Option
 			return false;
 		}
 
-		MLResult Result = MLMediaPlayerSetDataSourceForPath(MediaPlayerHandle, TCHAR_TO_UTF8(*FilePath));
+		Result = MLMediaPlayerSetDataSourceForPath(MediaPlayerHandle, TCHAR_TO_UTF8(*FilePath));
 		if (Result != MLResult_Ok)
 		{
 			UE_LOG(LogMagicLeapMedia, Error, TEXT("MLMediaPlayerSetDataSourceForPath for path %s failed with error %s."), *FilePath, UTF8_TO_TCHAR(MLMediaResultGetString(Result)));
@@ -399,10 +317,51 @@ bool FMagicLeapMediaPlayer::Open(const FString& Url, const IMediaOptions* Option
 			return false;
 		}
 	}
+	else if (Url.StartsWith(sharedFileSchema))
+	{
+		const FString SharedFileName = Url.RightChop(sharedFileSchema.Len());
+		const char* Filename_UTF8 = TCHAR_TO_UTF8(*SharedFileName);
+		MLSharedFileList* SharedFileList = nullptr;
+		Result = MLSharedFileRead(&Filename_UTF8, 1, &SharedFileList);
+		UE_CLOG(Result != MLResult_Ok, LogMagicLeapMedia, Error, TEXT("MLSharedFileRead(%s) failed with error %s"), *SharedFileName, UTF8_TO_TCHAR(MLSharedFileGetResultString(Result)));
+
+		if (Result == MLResult_Ok && SharedFileList != nullptr)
+		{
+			MLHandle ListLength = 0;
+			Result = MLSharedFileGetListLength(SharedFileList, &ListLength);
+			UE_CLOG(Result != MLResult_Ok, LogMagicLeapMedia, Error, TEXT("MLSharedFileGetListLength failed with error %s"), UTF8_TO_TCHAR(MLSharedFileGetResultString(Result)));
+			if (Result == MLResult_Ok && ListLength > 0)
+			{
+				MLFileInfo* FileInfo = nullptr;
+				Result = MLSharedFileGetMLFileInfoByIndex(SharedFileList, 0, &FileInfo);
+				UE_CLOG(Result != MLResult_Ok, LogMagicLeapMedia, Error, TEXT("MLSharedFileGetMLFileInfoByIndex failed with error %s"), UTF8_TO_TCHAR(MLSharedFileGetResultString(Result)));
+				if (Result == MLResult_Ok && FileInfo != nullptr)
+				{
+					MLFileDescriptor FileHandle = -1;
+					Result = MLFileInfoGetFD(FileInfo, &FileHandle);
+					UE_CLOG(Result != MLResult_Ok, LogMagicLeapMedia, Error, TEXT("MLFileInfoGetFD failed with error %s"), UTF8_TO_TCHAR(MLGetResultString(Result)));
+					if (Result == MLResult_Ok && FileHandle > -1)
+					{
+						Result = MLMediaPlayerSetDataSourceForFD(MediaPlayerHandle, FileHandle);
+						UE_CLOG(Result != MLResult_Ok, LogMagicLeapMedia, Error, TEXT("MLMediaPlayerSetDataSourceForFD failed with error %s"), UTF8_TO_TCHAR(MLMediaResultGetString(Result)));
+					}
+				}
+			}
+
+			MLSharedFileListRelease(&SharedFileList);
+		}
+
+		if (Result != MLResult_Ok)
+		{
+			UE_LOG(LogMagicLeapMedia, Error, TEXT("MLMediaPlayerSetDataSourceForFD for shared file %s failed."), *Url);
+			EventSink.ReceiveMediaEvent(EMediaEvent::MediaOpenFailed);
+			return false;
+		}
+	}
 	else
 	{    
 		// open remote media
-		MLResult Result = MLMediaPlayerSetDataSourceForURI(MediaPlayerHandle, TCHAR_TO_UTF8(*Url));
+		Result = MLMediaPlayerSetDataSourceForURI(MediaPlayerHandle, TCHAR_TO_UTF8(*Url));
 		if (Result != MLResult_Ok)
 		{
 			UE_LOG(LogMagicLeapMedia, Error, TEXT("MLMediaPlayerSetDataSourceForURI for remote media source %s failed with error %s."), *Url, UTF8_TO_TCHAR(MLMediaResultGetString(Result)));
@@ -411,12 +370,15 @@ bool FMagicLeapMediaPlayer::Open(const FString& Url, const IMediaOptions* Option
 		}
 	}
 
+	const bool Precache = (Options != nullptr) ? Options->GetMediaOption("PrecacheFile", false) : false;
+	UE_CLOG(Precache, LogMagicLeapMedia, Warning, TEXT("Precaching is not supported in Magic Leap Media. Use Magic Leap Media Codec instead if this feature is necessarry."));
+
 	EventSink.ReceiveMediaEvent(EMediaEvent::MediaConnecting);
 
 	// prepare media
 	MediaUrl = Url;
 
-	MLResult Result = MLMediaPlayerPrepareAsync(MediaPlayerHandle);
+	Result = MLMediaPlayerPrepareAsync(MediaPlayerHandle);
 	if (Result != MLResult_Ok)
 	{
 		UE_LOG(LogMagicLeapMedia, Error, TEXT("MLMediaPlayerPrepareAsync for media source %s failed with error %s"), *Url, UTF8_TO_TCHAR(MLMediaResultGetString(Result)));
@@ -522,7 +484,7 @@ bool FMagicLeapMediaPlayer::Seek(const FTimespan& Time)
 	}
 	else if (CurrentState == EMediaState::Playing || CurrentState == EMediaState::Paused)
 	{
-		MLResult Result = MLMediaPlayerSeekTo(MediaPlayerHandle, Time.GetTotalMilliseconds(), MLMediaSeekMode_Closest_Sync);
+		MLResult Result = MLMediaPlayerSeekTo(MediaPlayerHandle, static_cast<int32>(Time.GetTotalMilliseconds()), MLMediaSeekMode_Closest_Sync);
 		if (Result != MLResult_Ok)
 		{
 			bSuccess = false;
@@ -575,7 +537,16 @@ bool FMagicLeapMediaPlayer::SetRate(float Rate)
 	}
 	else if (Rate == 1.0f)
 	{
-		bResult = SetRateOne();
+		MLResult Result = MLMediaPlayerStart(MediaPlayerHandle);
+		if (Result != MLResult_Ok)
+		{
+			UE_LOG(LogMagicLeapMedia, Error, TEXT("MLMediaPlayerStart failed with error %s"), UTF8_TO_TCHAR(MLMediaResultGetString(Result)));
+			bResult = false;
+		}
+
+		CurrentState = EMediaState::Playing;
+		EventSink.ReceiveMediaEvent(EMediaEvent::PlaybackResumed);
+		bResult = true;
 	}
 	else
 	{
@@ -768,154 +739,44 @@ void FMagicLeapMediaPlayer::TickFetch(FTimespan DeltaTime, FTimespan Timecode)
 						return;
 					}
 
-					// Clear gl errors as they can creep in from the UE4 renderer.
-					glGetError();
+					int32 TextureID = *reinterpret_cast<int32*>(MediaVideoTexture->GetNativeResource());
 
-					if (!TextureDataPtr->bContextCreated)
-					{
-						TextureDataPtr->InitContext();
-						TextureDataPtr->bContextCreated = true;
-					}
-					TextureDataPtr->SaveContext();
-					TextureDataPtr->MakeCurrent();
+					TextureDataPtr->TextureRenderer.DestroyImageKHR();
 
-					int32 TextureID = *reinterpret_cast<int32*>(MediaVideoTexture->GetNativeResource());
-					if (TextureDataPtr->Image != EGL_NO_IMAGE_KHR)
-					{
-						eglDestroyImageKHR(eglGetCurrentDisplay(), TextureDataPtr->Image);
-						TextureDataPtr->Image = EGL_NO_IMAGE_KHR;
-					}
 					if (TextureDataPtr->PreviousNativeBuffer != 0 && MLHandleIsValid(TextureDataPtr->PreviousNativeBuffer))
 					{
 						Params.MediaPlayer->RenderThreadReleaseNativeBuffer(Params.MediaPlayerHandle, TextureDataPtr->PreviousNativeBuffer);
 					}
 					TextureDataPtr->PreviousNativeBuffer = nativeBuffer;
 
-					// Wrap latest decoded frame into a new gl texture oject
-					TextureDataPtr->Image = eglCreateImageKHR(TextureDataPtr->Display, EGL_NO_CONTEXT, EGL_NATIVE_BUFFER_ANDROID, (EGLClientBuffer)(void *)nativeBuffer, NULL);
-					if (TextureDataPtr->Image == EGL_NO_IMAGE_KHR)
-					{
-						EGLint errorcode = eglGetError();
-						UE_LOG(LogMagicLeapMedia, Error, TEXT("Failed to create EGLImage from the buffer. %d"), errorcode);
-						TextureDataPtr->RestoreContext();
-						return;
-					}
-					glActiveTexture(GL_TEXTURE0);
-					glBindTexture(GL_TEXTURE_EXTERNAL_OES, TextureID);
-					glEGLImageTargetTexture2DOES(GL_TEXTURE_EXTERNAL_OES, TextureDataPtr->Image);
-					glTexParameteri(GL_TEXTURE_EXTERNAL_OES, GL_TEXTURE_MIN_FILTER, GL_LINEAR);
-					glTexParameteri(GL_TEXTURE_EXTERNAL_OES, GL_TEXTURE_MAG_FILTER, GL_LINEAR);
-					glTexParameteri(GL_TEXTURE_EXTERNAL_OES, GL_TEXTURE_WRAP_S, GL_CLAMP_TO_EDGE);
-					glTexParameteri(GL_TEXTURE_EXTERNAL_OES, GL_TEXTURE_WRAP_T, GL_CLAMP_TO_EDGE);
-					glBindTexture(GL_TEXTURE_EXTERNAL_OES, 0);
-
-					TextureDataPtr->RestoreContext();
-
-					if (!TextureDataPtr->bIsVideoTextureValid)
+					if (TextureDataPtr->TextureRenderer.CreateImageKHR(reinterpret_cast<void*>(nativeBuffer)))
 					{
-						FSamplerStateInitializerRHI SamplerStateInitializer(SF_Bilinear, AM_Clamp, AM_Clamp, AM_Clamp);
-						FSamplerStateRHIRef SamplerStateRHI = RHICreateSamplerState(SamplerStateInitializer);
-						Params.MediaPlayer->RegisterExternalTexture(Params.PlayerGuid, MediaVideoTexture, SamplerStateRHI);
-						TextureDataPtr->bIsVideoTextureValid = true;
+						TextureDataPtr->TextureRenderer.BindImageKHRToTexture(TextureID);
+						if (!TextureDataPtr->bIsVideoTextureValid)
+						{
+							FSamplerStateInitializerRHI SamplerStateInitializer(SF_Bilinear, AM_Clamp, AM_Clamp, AM_Clamp);
+							FSamplerStateRHIRef SamplerStateRHI = RHICreateSamplerState(SamplerStateInitializer);
+							Params.MediaPlayer->RegisterExternalTexture(Params.PlayerGuid, MediaVideoTexture, SamplerStateRHI);
+							TextureDataPtr->bIsVideoTextureValid = true;
+						}
 					}
 				});
 		}
 	}
-	else
-	{
-		FMediaVideoTrackFormat TrackFormat;
-		bool bTrackFormatValid = GetVideoTrackFormat(0, 0, TrackFormat);
-		if (!bTrackFormatValid)
-		{
-			UE_LOG(LogMagicLeapMedia, Error, TEXT("Invalid track format. Skipping frame."));
-			return;
-		}
-
-		auto VideoSample = VideoSamplePool->AcquireShared();
-		if (!VideoSample->Initialize(TrackFormat.Dim, FTimespan::FromSeconds(1.0 / TrackFormat.FrameRate)))
-		{
-			UE_LOG(LogMagicLeapMedia, Error, TEXT("Could not initialize video sample."));
-			return;
-		}
-
-		struct FWriteVideoSampleParams
-		{
-			FMagicLeapMediaPlayer* MediaPlayer;
-			TWeakPtr<FMagicLeapVideoTextureDataGL, ESPMode::ThreadSafe> TextureData;
-			MLHandle MediaPlayerHandle;
-			TWeakPtr<FMediaSamples, ESPMode::ThreadSafe> SamplesPtr;
-			TSharedRef<FMagicLeapMediaTextureSample, ESPMode::ThreadSafe> VideoSample;
-			FCriticalSection* CriticalSectionPtr;
-			FMediaWorker* MediaWorkerPtr;
-		}
-		Params = { this, StaticCastSharedPtr<FMagicLeapVideoTextureDataGL>(TextureData), MediaPlayerHandle, Samples, VideoSample, &CriticalSection, MediaWorker };
-
-		ENQUEUE_RENDER_COMMAND(MagicLeapMediaPlayerWriteVideoSample)(
-			[Params](FRHICommandListImmediate& RHICmdList)
-			{
-				auto PinnedTextureData = Params.TextureData.Pin();
-				auto PinnedSamples = Params.SamplesPtr.Pin();
 
-				if (!PinnedTextureData.IsValid() || !PinnedSamples.IsValid())
-				{
-					UE_LOG(LogMagicLeapMedia, Error, TEXT("Invalid texture data or samples."));
-					return;
-				}
-
-				if (FLuminPlatformMisc::ShouldUseVulkan() || FLuminPlatformMisc::ShouldUseDesktopOpenGL())
-				{
-					bool BufferAvailable = false;
-					{
-						FScopeLock Lock(Params.CriticalSectionPtr);
-						if (Params.MediaWorkerPtr->NextBufferAvailable.GetValue() > 0)
-						{
-							FIntPoint CurrentDim;
-							FTimespan CurrentTime;
-							void* ReadBuffer = Params.MediaWorkerPtr->GetReadBuffer(&CurrentDim, &CurrentTime);
-							Params.VideoSample->Initialize(CurrentDim, Params.VideoSample->GetDuration());
-							Params.VideoSample->InitializeBuffer(CurrentTime, ReadBuffer, true);
-							Params.MediaWorkerPtr->NextBufferAvailable.Decrement();
-							BufferAvailable = true;
-						}
-					}
-
-					if (BufferAvailable)
-					{
-						PinnedSamples->AddVideo(Params.VideoSample);
-					}
-				}
-				else
-				{
-					if (!Params.MediaPlayer->RenderThreadIsBufferAvailable(Params.MediaPlayerHandle))
-					{
-						return;
-					}
+	if (GetMediaPlayerState(MLMediaPlayerPollingStateFlag_HasSubtitleUpdated))
+	{
+		MLCea608CaptionSegment Segment;
+		MLCea608CaptionSegmentInit(&Segment);
+		MLMediaPlayerSubtitleData* SubtitleDataPtr = nullptr;
 
-					// MLHandle because Unreal's uint64 is 'unsigned long long *' whereas uint64_t for the C-API is 'unsigned long *'
-					// TODO: Fix the Unreal types for the above comment.
-					MLHandle NativeBuffer = ML_INVALID_HANDLE;
-					if (!Params.MediaPlayer->RenderThreadGetNativeBuffer(Params.MediaPlayerHandle, NativeBuffer, PinnedTextureData->bIsVideoTextureValid))
-					{
-						UE_LOG(LogMagicLeapMedia, Error, TEXT("Error acquiring next video buffer"));
-						return;
-					}
+		ProcessCaptioningSegment(Segment, SubtitleDataPtr);
 
-					// write frame into texture
-					FRHITexture2D* Texture = Params.VideoSample->InitializeTexture(Params.MediaPlayer->GetTime());
-					if (Texture != nullptr)
-					{
-						int32 Resource = *reinterpret_cast<int32*>(Texture->GetNativeResource());
-						bool TextureCopyResult = PinnedTextureData->ExternalRenderer->CopyFrameTexture(Resource, NativeBuffer, Params.VideoSample->GetDim(), nullptr);
-						if (!TextureCopyResult)
-						{
-							UE_LOG(LogMagicLeapMedia, Error, TEXT("CopyFrameTexture failed"));
-						}
-					}
+		MLResult ReleaseSegmentResult = MLMediaCCParserReleaseSegment(CaptionParserHandle, &Segment);
+		UE_CLOG(ReleaseSegmentResult != MLResult_Ok, LogMagicLeapMedia, Error, TEXT("MLMediaCCParserReleaseSegment failed with error %s."), UTF8_TO_TCHAR(MLGetResultString(ReleaseSegmentResult)));
 
-					Params.MediaPlayer->RenderThreadReleaseNativeBuffer(Params.MediaPlayerHandle, NativeBuffer);
-					PinnedSamples->AddVideo(Params.VideoSample);
-				}
-			});
+		MLResult ReleaseSubtitleResult = MLMediaPlayerReleaseSubtitleEx(MediaPlayerHandle);
+		UE_CLOG(ReleaseSubtitleResult != MLResult_Ok, LogMagicLeapMedia, Error, TEXT("MLMediaPlayerReleaseSubtitleEx failed with error %s."), UTF8_TO_TCHAR(MLGetResultString(ReleaseSubtitleResult)));
 	}
 }
 
@@ -979,11 +840,6 @@ void FMagicLeapMediaPlayer::TickInput(FTimespan DeltaTime, FTimespan Timecode)
 			// notify listeners
 			EventSink.ReceiveMediaEvent(EMediaEvent::TracksChanged);
 			EventSink.ReceiveMediaEvent(EMediaEvent::MediaOpened);
-
-			if (!GSupportsImageExternal)
-			{
-				MediaWorker->InitThread();
-			}
 		}
 		else
 		{
@@ -991,6 +847,17 @@ void FMagicLeapMediaPlayer::TickInput(FTimespan DeltaTime, FTimespan Timecode)
 		}
 	}
 
+	// Subtitle/Caption tracks are only picked up on MetadataUpdated
+	if (GetMediaPlayerState(MLMediaPlayerPollingStateFlag_HasMetadataUpdated))
+	{
+		CheckSubtitlesAndCaptioning();
+		if (!MLHandleIsValid(CaptionParserHandle))
+		{
+			MLResult Result = MLMediaCCParserCreate(&CaptionParserHandle);
+			UE_CLOG(Result != MLResult_Ok, LogMagicLeapMedia, Error, TEXT("MLMediaCCParserCreate failed with error %s."), UTF8_TO_TCHAR(MLGetResultString(Result)));
+		}
+	}
+
 	if (GetMediaPlayerState(MLMediaPlayerPollingStateFlag_HasSeekCompleted))
 	{
 		EventSink.ReceiveMediaEvent(EMediaEvent::SeekCompleted);
@@ -1043,8 +910,37 @@ void FMagicLeapMediaPlayer::TickInput(FTimespan DeltaTime, FTimespan Timecode)
 
 bool FMagicLeapMediaPlayer::GetAudioTrackFormat(int32 TrackIndex, int32 FormatIndex, FMediaAudioTrackFormat& OutFormat) const
 {
-	// TODO: can we implement this using the audio handle coming from MLMediaPlayerGetAudioHandle()?
-	return false;
+	if ((FormatIndex != static_cast<int32>(EMediaTrackType::Audio)) || TrackIndex >= TrackInfo[EMediaTrackType::Audio].Num() || TrackIndex < 0)
+	{
+		return false;
+	}
+
+	MLHandle FormatHandle;
+	MLResult Result = MLMediaPlayerGetTrackMediaFormat(MediaPlayerHandle, static_cast<uint32>(TrackInfo[EMediaTrackType::Audio][TrackIndex]), &FormatHandle);
+	if (Result != MLResult_Ok)
+	{
+		UE_LOG(LogMagicLeapMedia, Error, TEXT("MLMediaPlayerGetTrackMediaFormat failed with error %s"), UTF8_TO_TCHAR(MLMediaResultGetString(Result)));
+		return false;
+	}
+
+	//Update integer format values
+	Result = MLMediaFormatGetKeyValueInt32(FormatHandle, MLMediaFormat_Key_Sample_Rate, &OutFormat.SampleRate);
+	UE_CLOG(Result != MLResult_Ok, LogMagicLeapMedia, Error, TEXT("MLMediaFormatGetKeyValueInt32(sample-rate) failed with error %s"), UTF8_TO_TCHAR(MLMediaResultGetString(Result)));
+	Result = MLMediaFormatGetKeyValueInt32(FormatHandle, MLMediaFormat_Key_Channel_Count, &OutFormat.NumChannels);
+	UE_CLOG(Result != MLResult_Ok, LogMagicLeapMedia, Error, TEXT("MLMediaFormatGetKeyValueInt32(channel-count) failed with error %s"), UTF8_TO_TCHAR(MLMediaResultGetString(Result)));
+	Result = MLMediaFormatGetKeyValueInt32(FormatHandle, MLMediaFormat_Key_Bit_Rate, &OutFormat.BitsPerSample);
+	UE_CLOG(Result != MLResult_Ok, LogMagicLeapMedia, Error, TEXT("MLMediaFormatGetKeyValueInt32(bit-rate) failed with error %s"), UTF8_TO_TCHAR(MLMediaResultGetString(Result)));
+
+	//Update string format value, removing the format type.
+	char FormatName[MAX_FORMAT_STRING_SIZE] = "";
+	Result = MLMediaFormatGetKeyString(FormatHandle, MLMediaFormat_Key_Mime, &FormatName);
+	UE_CLOG(Result != MLResult_Ok, LogMagicLeapMedia, Error, TEXT("MLMediaFormatGetKeyString(mime) failed with error %s"), UTF8_TO_TCHAR(MLMediaResultGetString(Result)));
+	OutFormat.TypeName = (Result == MLResult_Ok) ? FString(FormatName) : FString();
+	OutFormat.TypeName.RemoveFromStart("audio/");
+
+	MLMediaFormatDestroy(FormatHandle);
+
+	return true;
 }
 
 int32 FMagicLeapMediaPlayer::GetNumTracks(EMediaTrackType TrackType) const
@@ -1125,11 +1021,34 @@ bool FMagicLeapMediaPlayer::GetVideoTrackFormat(int32 TrackIndex, int32 FormatIn
 		return false;
 	}
 
-	OutFormat.Dim = FIntPoint(width, height);
-	// TODO: Don't hardcode. Get from C-API. ml_media_player api does not provide that right now. Try the ml_media_codec api.
-	OutFormat.FrameRate = 30.0f;
-	OutFormat.FrameRates = TRange<float>(30.0f);
-	OutFormat.TypeName = TEXT("BGRA");
+	MLHandle MediaFormatHandle = ML_INVALID_HANDLE;
+	Result = MLMediaPlayerGetTrackMediaFormat(MediaPlayerHandle, TrackIndex, &MediaFormatHandle);
+	if (Result != MLResult_Ok)
+	{
+		UE_LOG(LogMagicLeapMedia, Error, TEXT("MLMediaPlayerGetTrackMediaFormat failed with error %s"), UTF8_TO_TCHAR(MLMediaResultGetString(Result)));
+		return false;
+	}
+
+	int32 FrameRate = 0;
+	MLMediaFormatGetKeyValueInt32(MediaFormatHandle, MLMediaFormat_Key_Frame_Rate, &FrameRate);
+
+	if (Result != MLResult_Ok)
+	{
+		UE_LOG(LogMagicLeapMedia, Error, TEXT("MLMediaFormatGetKeyValueInt32 failed with error %s"), UTF8_TO_TCHAR(MLMediaResultGetString(Result)));
+		return false;
+	}
+	OutFormat.FrameRate = static_cast<float>(FrameRate);
+	OutFormat.FrameRates = TRange<float>(FrameRate);
+
+	char MIME[MAX_KEY_STRING_SIZE] = "";
+	Result = MLMediaFormatGetKeyString(MediaFormatHandle, MLMediaFormat_Key_Mime, &MIME);
+	if (Result != MLResult_Ok)
+	{
+		UE_LOG(LogMagicLeapMedia, Error, TEXT("MLMediaFormatGetKeyString failed with error %s"), UTF8_TO_TCHAR(MLMediaResultGetString(Result)));
+		return false;
+	}
+	OutFormat.TypeName = UTF8_TO_TCHAR(MIME);
+
 	return true;
 }
 
@@ -1164,7 +1083,7 @@ void FMagicLeapMediaPlayer::HandleApplicationHasEnteredForeground()
 	// check state in case changed before ticked
 	if (CurrentState == EMediaState::Paused && bWasMediaPlayingBeforeAppPause)
 	{
-		// pause
+		// Unpause
 		SetRate(1.0f);
 	}
 }
@@ -1180,6 +1099,85 @@ void FMagicLeapMediaPlayer::HandleApplicationWillEnterBackground()
 	}
 }
 
+void FMagicLeapMediaPlayer::CheckSubtitlesAndCaptioning()
+{
+	uint32 NumTracks = 0;
+	MLMediaPlayerGetTrackCount(MediaPlayerHandle, &NumTracks);
+	for (uint32 i = 0; i < NumTracks; ++i)
+	{
+		MLMediaPlayerTrackType TrackType = MediaPlayerTrackType_Unknown;
+		MLMediaPlayerGetTrackType(MediaPlayerHandle, i, &TrackType);
+		switch (TrackType)
+		{
+		case MediaPlayerTrackType_TimedText:
+			if (!TrackInfo[EMediaTrackType::Caption].Contains(static_cast<int32>(i)))
+			{
+				TrackInfo[EMediaTrackType::Caption].Add(static_cast<int32>(i));
+			}
+			if (SelectedTrack[EMediaTrackType::Caption] == INDEX_NONE)
+			{
+				SelectedTrack[EMediaTrackType::Caption] = 0;
+				SelectTrack(EMediaTrackType::Caption, SelectedTrack[EMediaTrackType::Caption]);
+			}
+			break;
+		case MediaPlayerTrackType_Subtitle:
+			if (!TrackInfo[EMediaTrackType::Subtitle].Contains(static_cast<int32>(i)))
+			{
+				TrackInfo[EMediaTrackType::Subtitle].Add(static_cast<int32>(i));
+			}
+			if (SelectedTrack[EMediaTrackType::Subtitle] == INDEX_NONE)
+			{
+				SelectedTrack[EMediaTrackType::Subtitle] = 0;
+				SelectTrack(EMediaTrackType::Subtitle, SelectedTrack[EMediaTrackType::Subtitle]);
+			}
+			break;
+		// We currently ignore other track types at detected metadata changes
+		default:
+			break;
+		}
+	}
+}
+
+void FMagicLeapMediaPlayer::ProcessCaptioningSegment(MLCea608CaptionSegment& Segment, MLMediaPlayerSubtitleData* SubtitleDataPtr)
+{
+	check(MLHandleIsValid(CaptionParserHandle));
+	MLResult Result = MLMediaPlayerGetSubtitleEx(MediaPlayerHandle, &SubtitleDataPtr);
+	if (Result == MLResult_Ok)
+	{
+		MLResult DisplayableResult = MLMediaCCParserGetDisplayable(CaptionParserHandle, SubtitleDataPtr->data, SubtitleDataPtr->data_size, &Segment);
+		if (DisplayableResult != MLResult_Ok)
+		{
+			UE_LOG(LogMagicLeapMedia, Error, TEXT("MLMediaCCParserGetDisplayable failed with error %s."), UTF8_TO_TCHAR(MLGetResultString(DisplayableResult)));
+			return;
+		}
+
+		for (uint32 Index = 0; Index < MLCea608_CCMaxRowsPlus2; ++Index)
+		{
+			if (Segment.ccLine[Index] != nullptr)
+			{
+				TArray<char> LineCharacters;
+				LineCharacters.Reserve(MLCea608_CCMaxColsPlus2);
+				// Caption lines are kept in their respective row and column due to style data corresponding
+				// to row/column. So we must check for non-characters in the string
+				for (uint32 LineIndex = 0; LineIndex < MLCea608_CCMaxColsPlus2; ++LineIndex)
+				{
+					if (Segment.ccLine[Index]->displayChars[LineIndex] != '\xff')
+					{
+						LineCharacters.Emplace(Segment.ccLine[Index]->displayChars[LineIndex]);
+					}
+				}
+				TSharedRef<FMagicLeapMediaOverlaySample, ESPMode::ThreadSafe> Subtitle = MakeShared<FMagicLeapMediaOverlaySample, ESPMode::ThreadSafe>();
+				Subtitle->Intitialize(LineCharacters.GetData(), EMediaOverlaySampleType::Subtitle, FTimespan::FromMicroseconds(SubtitleDataPtr->duration_us), FTimespan::FromMicroseconds(SubtitleDataPtr->start_time_us));
+				Samples->AddSubtitle(Subtitle);
+			}
+		}
+	}
+	else
+	{
+		UE_LOG(LogMagicLeapMedia, Error, TEXT("MLMediaPlayerGetSubtitleEx failed with error %s."), UTF8_TO_TCHAR(MLGetResultString(Result)));
+	}
+}
+
 FIntPoint FMagicLeapMediaPlayer::GetVideoDimensions() const
 {
 	int32 width = 0;
@@ -1192,20 +1190,6 @@ FIntPoint FMagicLeapMediaPlayer::GetVideoDimensions() const
 	return FIntPoint(width, height);
 }
 
-bool FMagicLeapMediaPlayer::SetRateOne()
-{
-	MLResult Result = MLMediaPlayerStart(MediaPlayerHandle);
-	if (Result != MLResult_Ok)
-	{
-		UE_LOG(LogMagicLeapMedia, Error, TEXT("MLMediaPlayerStart failed with error %s"), UTF8_TO_TCHAR(MLMediaResultGetString(Result)));
-		return false;
-	}
-	
-	CurrentState = EMediaState::Playing;
-	EventSink.ReceiveMediaEvent(EMediaEvent::PlaybackResumed);
-	return true;
-}
-
 bool FMagicLeapMediaPlayer::GetMediaPlayerState(uint16 FlagToPoll) const
 {
 	uint16_t StateFlags = 0;
@@ -1291,4 +1275,33 @@ bool FMagicLeapMediaPlayer::RenderThreadGetCurrentPosition(const MLHandle InMedi
 	return true;
 }
 
+void FMagicLeapMediaPlayer::TriggerResetAndDestroy()
+{
+	// Platform api for MLMediaPlayerReset()/Destroy() has a bug where the function can stall for several seconds.
+	// Until that is fixed, we call those functions on a detached thread so as not to block the app's process from exiting.
+
+	int32 ThreadFuncResult = 0;
+	pthread_t Thread;
+	pthread_attr_t attr;
+
+	ThreadFuncResult = pthread_attr_init(&attr);
+	if (ThreadFuncResult == 0)
+	{
+		ThreadFuncResult = pthread_attr_setdetachstate(&attr, PTHREAD_CREATE_DETACHED);
+		if (ThreadFuncResult == 0)
+		{
+			ThreadFuncResult = pthread_create(&Thread, &attr, _ThreadProc, reinterpret_cast<void*>(MediaPlayerHandle));
+			UE_CLOG(ThreadFuncResult != 0, LogMagicLeapMedia, Error, TEXT("pthread_create failed with error %d."), ThreadFuncResult);
+		}
+		else
+		{
+			UE_LOG(LogMagicLeapMedia, Error, TEXT("pthread_attr_setdetachstate failed with error %d."), ThreadFuncResult);
+		}
+	}
+	else
+	{
+		UE_LOG(LogMagicLeapMedia, Error, TEXT("pthread_attr_init failed with error %d."), ThreadFuncResult);
+	}
+}
+
 #undef LOCTEXT_NAMESPACE
diff --git a/Engine/Plugins/Lumin/MagicLeapMedia/Source/MagicLeapMedia/Private/MagicLeapMediaPlayer.h b/Engine/Plugins/Lumin/MagicLeapMedia/Source/MagicLeapMedia/Private/MagicLeapMediaPlayer.h
--- a/Engine/Plugins/Lumin/MagicLeapMedia/Source/MagicLeapMedia/Private/MagicLeapMediaPlayer.h
+++ b/Engine/Plugins/Lumin/MagicLeapMedia/Source/MagicLeapMedia/Private/MagicLeapMediaPlayer.h
@@ -16,12 +16,11 @@
 #include "RenderingThread.h"
 #include "RHI.h"
 #include "RHIResources.h"
-
-#include "ml_api.h"
+#include "Lumin/CAPIShims/LuminAPIMediaCCParser.h"
+#include "Lumin/CAPIShims/LuminAPIMediaPlayer.h"
 
 class FMediaSamples;
 class IMediaEventSink;
-class FMagicLeapMediaTextureSamplePool;
 struct FMagicLeapVideoTextureData;
 
 /**
@@ -30,7 +29,7 @@ struct FMagicLeapVideoTextureData;
 class FMagicLeapMediaPlayer : public IMediaPlayer, public IMediaControls, public IMediaCache, public IMediaTracks, public IMediaView
 {
 public:
-  
+
 	/**
 	 * Create and initialize a new instance.
 	 *
@@ -100,7 +99,14 @@ protected:
 	// TODO: mature this to return track dimensions.
 	FIntPoint GetVideoDimensions() const;
 
+	/** Adds Subtitle and Caption tracks, should only be used when metadata has changed */
+	void CheckSubtitlesAndCaptioning();
+
+	/** Processes segments of captioned text into media overlay samples */
+	void ProcessCaptioningSegment(MLCea608CaptionSegment& Segment, MLMediaPlayerSubtitleData* SubtitleDataPtr);
+
 	MLHandle MediaPlayerHandle;
+	MLHandle CaptionParserHandle;
 
 	bool bMediaPrepared;
 
@@ -137,17 +143,12 @@ protected:
 	/** The media sample queue. */
 	TSharedPtr<FMediaSamples, ESPMode::ThreadSafe> Samples;
 
-	/** Video sample object pool. */
-	FMagicLeapMediaTextureSamplePool* VideoSamplePool;
-
 	TSharedPtr<FMagicLeapVideoTextureData, ESPMode::ThreadSafe> TextureData;
 
 	TMap<EMediaTrackType, TArray<int32>> TrackInfo;
 
 	TMap<EMediaTrackType, int32> SelectedTrack;
 
-	class FMediaWorker* MediaWorker;
-
 	FCriticalSection CriticalSection;
 
 	bool bWasMediaPlayingBeforeAppPause;
@@ -157,11 +158,11 @@ protected:
 	TAtomic<FTimespan> CurrentPlaybackTime;
 
 private:
-	virtual bool SetRateOne();
-	virtual bool GetMediaPlayerState(uint16 FlagToPoll) const;
-	virtual void RegisterExternalTexture(const FGuid& InGuid, FTextureRHIRef& InTextureRHI, FSamplerStateRHIRef& InSamplerStateRHI);
-	virtual bool RenderThreadIsBufferAvailable(MLHandle MediaPlayerHandle);
-	virtual bool RenderThreadGetNativeBuffer(const MLHandle MediaPlayerHandle, MLHandle& NativeBuffer, bool& OutIsVideoTextureValid);
-	virtual bool RenderThreadReleaseNativeBuffer(const MLHandle MediaPlayerHandle, MLHandle NativeBuffer);
-	virtual bool RenderThreadGetCurrentPosition(const MLHandle MediaPlayerHandle, int32& CurrentPosition);
+	bool GetMediaPlayerState(uint16 FlagToPoll) const;
+	void RegisterExternalTexture(const FGuid& InGuid, FTextureRHIRef& InTextureRHI, FSamplerStateRHIRef& InSamplerStateRHI);
+	bool RenderThreadIsBufferAvailable(MLHandle MediaPlayerHandle);
+	bool RenderThreadGetNativeBuffer(const MLHandle MediaPlayerHandle, MLHandle& NativeBuffer, bool& OutIsVideoTextureValid);
+	bool RenderThreadReleaseNativeBuffer(const MLHandle MediaPlayerHandle, MLHandle NativeBuffer);
+	bool RenderThreadGetCurrentPosition(const MLHandle MediaPlayerHandle, int32& CurrentPosition);
+	void TriggerResetAndDestroy();
 };
diff --git a/Engine/Plugins/Lumin/MagicLeapMedia/Source/MagicLeapMedia/Public/IMagicLeapMediaModule.h b/Engine/Plugins/Lumin/MagicLeapMedia/Source/MagicLeapMedia/Public/IMagicLeapMediaModule.h
--- a/Engine/Plugins/Lumin/MagicLeapMedia/Source/MagicLeapMedia/Public/IMagicLeapMediaModule.h
+++ b/Engine/Plugins/Lumin/MagicLeapMedia/Source/MagicLeapMedia/Public/IMagicLeapMediaModule.h
@@ -21,8 +21,6 @@ public:
 	 */
 	virtual TSharedPtr<IMediaPlayer, ESPMode::ThreadSafe> CreatePlayer(class IMediaEventSink& EventSink) = 0;
 
-	virtual TSharedPtr<IMediaPlayer, ESPMode::ThreadSafe> CreateCameraPreviewPlayer(class IMediaEventSink& EventSink) = 0;
-
 public:
 	/** Virtual destructor. */
 	virtual ~IMagicLeapMediaModule() { }
diff --git a/Engine/Plugins/Lumin/MagicLeapMedia/Source/MagicLeapMediaCodec/MagicLeapMediaCodec.Build.cs b/Engine/Plugins/Lumin/MagicLeapMedia/Source/MagicLeapMediaCodec/MagicLeapMediaCodec.Build.cs
--- a/Engine/Plugins/Lumin/MagicLeapMedia/Source/MagicLeapMediaCodec/MagicLeapMediaCodec.Build.cs
+++ b/Engine/Plugins/Lumin/MagicLeapMedia/Source/MagicLeapMediaCodec/MagicLeapMediaCodec.Build.cs
@@ -8,8 +8,7 @@ namespace UnrealBuildTool.Rules
 	{
 		public MagicLeapMediaCodec(ReadOnlyTargetRules Target) : base(Target)
 		{
-            PublicIncludePaths.Add(Path.Combine(ModuleDirectory, "Public"));
-			string EngineDir = Path.GetFullPath(Target.RelativeEnginePath);
+			PublicIncludePaths.Add(Path.Combine(ModuleDirectory, "Public"));
 
 			DynamicallyLoadedModuleNames.AddRange(
 				new string[] {
@@ -22,12 +21,11 @@ namespace UnrealBuildTool.Rules
 					"Engine",
 					"RenderCore",
 					"MediaUtils",
-					"OpenGLDrv",
 					"LuminRuntimeSettings",
 					"MLSDK",
-                    "MagicLeap",
-                    "MagicLeapHelperVulkan",
-            });
+					"MagicLeapHelperVulkan",
+					"MagicLeapHelperOpenGL"
+			});
 
 			PrivateIncludePathModuleNames.AddRange(
 				new string[] {
@@ -36,11 +34,8 @@ namespace UnrealBuildTool.Rules
 
 			PrivateIncludePaths.AddRange(
 				new string[] {
-                    "MagicLeapMedia/Private",
-					Path.Combine(EngineDir, "Source/Runtime/OpenGLDrv/Private")
-            });
-
-			AddEngineThirdPartyPrivateStaticDependencies(Target, "OpenGL");
+					"MagicLeapMedia/Private"
+			});
 		}
 	}
 }
diff --git a/Engine/Plugins/Lumin/MagicLeapMedia/Source/MagicLeapMediaCodec/Private/MagicLeapMediaCodecPlayer.cpp b/Engine/Plugins/Lumin/MagicLeapMedia/Source/MagicLeapMediaCodec/Private/MagicLeapMediaCodecPlayer.cpp
--- a/Engine/Plugins/Lumin/MagicLeapMedia/Source/MagicLeapMediaCodec/Private/MagicLeapMediaCodecPlayer.cpp
+++ b/Engine/Plugins/Lumin/MagicLeapMedia/Source/MagicLeapMediaCodec/Private/MagicLeapMediaCodecPlayer.cpp
@@ -13,17 +13,14 @@
 #include "UObject/UObjectGlobals.h"
 #include "ExternalTexture.h"
 #include "MagicLeapMediaAudioSample.h"
-#include "Lumin/LuminEGL.h"
 #include "MagicLeapHelperVulkan.h"
-
-#if WITH_MLSDK
-#include "ml_version.h"
-#include "ml_media_extractor.h"
-#include "ml_media_codec.h"
-#include "ml_media_codeclist.h"
-#include "ml_media_format.h"
-#include "ml_media_error.h"
-#endif // WITH_MLSDK
+#include "ExternalTextureRendererGL.h"
+#include "Lumin/CAPIShims/LuminAPIMediaError.h"
+#include "Lumin/CAPIShims/LuminAPIMediaExtractor.h"
+#include "Lumin/CAPIShims/LuminAPIMediaFormat.h"
+#include "Lumin/CAPIShims/LuminAPIFileInfo.h"
+#include "Lumin/CAPIShims/LuminAPISharedFile.h"
+#include "Lumin/CAPIShims/LuminAPIMediaDataSource.h"
 
 #if MLSDK_VERSION_MAJOR == 0 && MLSDK_VERSION_MINOR < 18
 MLMediaFormatKey MLMediaFormat_Key_Mime = "mime";
@@ -92,76 +89,9 @@ struct FMagicLeapVideoTextureDataGL : public FMagicLeapVideoTextureData
 public:
 
 	FMagicLeapVideoTextureDataGL()
-		: Image(EGL_NO_IMAGE_KHR)
-		, Display(EGL_NO_DISPLAY)
-		, Context(EGL_NO_CONTEXT)
-		, SavedDisplay(EGL_NO_DISPLAY)
-		, SavedContext(EGL_NO_CONTEXT)
-		, bContextCreated(false)
 	{}
 
-	~FMagicLeapVideoTextureDataGL()
-	{
-		PreviousNativeBuffer = ML_INVALID_HANDLE;
-		eglDestroyContext(Display, Context);
-		Display = EGL_NO_DISPLAY;
-		Context = EGL_NO_CONTEXT;
-	}
-
-	bool InitContext()
-	{
-#if !PLATFORM_LUMINGL4
-		if (Context == EGL_NO_CONTEXT)
-		{
-			Display = LuminEGL::GetInstance()->GetDisplay();
-			EGLContext SharedContext = LuminEGL::GetInstance()->GetCurrentContext();
-			Context = SharedContext; // LuminEGL::GetInstance()->CreateContext(SharedContext);
-		}
-
-		return (Context != EGL_NO_CONTEXT);
-#else
-		return false;
-#endif
-	}
-
-	void SaveContext()
-	{
-#if !PLATFORM_LUMINGL4
-		SavedDisplay = LuminEGL::GetInstance()->GetDisplay();
-		SavedContext = LuminEGL::GetInstance()->GetCurrentContext();
-#endif
-	}
-
-	void MakeCurrent()
-	{
-#if !PLATFORM_LUMINGL4
-		return;	// skip for now
-		EGLBoolean bResult = eglMakeCurrent(Display, EGL_NO_SURFACE, EGL_NO_SURFACE, Context);
-		if (bResult == EGL_FALSE)
-		{
-			UE_LOG(LogMagicLeapMediaCodec, Error, TEXT("Error setting media player context."));
-		}
-#endif
-	}
-
-	void RestoreContext()
-	{
-#if !PLATFORM_LUMINGL4
-		return;	// skip for now
-		EGLBoolean bResult = eglMakeCurrent(SavedDisplay, EGL_NO_SURFACE, EGL_NO_SURFACE, SavedContext);
-		if (bResult == EGL_FALSE)
-		{
-			UE_LOG(LogMagicLeapMediaCodec, Error, TEXT("Error setting unreal context."));
-		}
-#endif
-	}
-
-	EGLImageKHR Image;
-	EGLDisplay Display;
-	EGLContext Context;
-	EGLDisplay SavedDisplay;
-	EGLContext SavedContext;
-	bool bContextCreated;
+	FExternalTextureRendererGL TextureRenderer;
 };
 
 FMagicLeapMediaCodecPlayer::FMagicLeapMediaCodecPlayer(IMediaEventSink& InEventSink)
@@ -176,6 +106,7 @@ FMagicLeapMediaCodecPlayer::FMagicLeapMediaCodecPlayer(IMediaEventSink& InEventS
 	, AudioSamplePool(new FMagicLeapMediaAudioSamplePool())
 	, bWasMediaPlayingBeforeAppPause(false)
 	, CurrentPlaybackTime(FTimespan::Zero())
+	, bJustSeeked(false)
 	, bLoopPlayback(false)
 	, bPlaybackCompleted(false)
 	, bPlaybackCompleted_RenderThread(false)
@@ -283,16 +214,9 @@ void FMagicLeapMediaCodecPlayer::Close()
 
 						// @todo: this causes a crash
 						//Params.TextureData->VideoTexture->Release();
-						Params.TextureData->SaveContext();
-						Params.TextureData->MakeCurrent();
 
-						if (Params.TextureData->Image != EGL_NO_IMAGE_KHR)
-						{
-							eglDestroyImageKHR(eglGetCurrentDisplay(), Params.TextureData->Image);
-							Params.TextureData->Image = EGL_NO_IMAGE_KHR;
-						}
+						Params.TextureData->TextureRenderer.DestroyImageKHR();
 
-						Params.TextureData->RestoreContext();
 						if (Params.TextureData->PreviousNativeBuffer != 0 && MLHandleIsValid(Params.TextureData->PreviousNativeBuffer))
 						{
 							Params.MediaPlayer->ReleaseNativeBuffer_RenderThread(Params.VideoCodecHandle, Params.TextureData->PreviousNativeBuffer);
@@ -350,6 +274,7 @@ void FMagicLeapMediaCodecPlayer::Close()
 	AudioSamplePool->Reset();
 	TrackInfo.Empty();
 	SelectedTrack.Empty();
+	bJustSeeked = false;
 
 	// notify listeners
 	EventSink.ReceiveMediaEvent(EMediaEvent::TracksChanged);
@@ -427,6 +352,7 @@ bool FMagicLeapMediaCodecPlayer::Open(const FString& Url, const IMediaOptions* O
 	MediaUrl = Url;
 
 	const FString localFileSchema = "file://";
+	const FString sharedFileSchema = "mlshared://";
 
 	// open the media
 	if (Url.StartsWith(localFileSchema))
@@ -503,6 +429,49 @@ bool FMagicLeapMediaCodecPlayer::Open(const FString& Url, const IMediaOptions* O
 			}
 		}
 	}
+	else if (Url.StartsWith(sharedFileSchema))
+	{
+		const FString SharedFileName = Url.RightChop(sharedFileSchema.Len());
+		const char* Filename_UTF8 = TCHAR_TO_UTF8(*SharedFileName);
+		MLSharedFileList* SharedFileList = nullptr;
+		MLResult Result = MLSharedFileRead(&Filename_UTF8, 1, &SharedFileList);
+		UE_CLOG(Result != MLResult_Ok, LogMagicLeapMediaCodec, Error, TEXT("MLSharedFileRead(%s) failed with error %s"), *SharedFileName, UTF8_TO_TCHAR(MLSharedFileGetResultString(Result)));
+
+		if (Result == MLResult_Ok && SharedFileList != nullptr)
+		{
+			MLHandle ListLength = 0;
+			Result = MLSharedFileGetListLength(SharedFileList, &ListLength);
+			UE_CLOG(Result != MLResult_Ok, LogMagicLeapMediaCodec, Error, TEXT("MLSharedFileGetListLength failed with error %s"), UTF8_TO_TCHAR(MLSharedFileGetResultString(Result)));
+			if (Result == MLResult_Ok && ListLength > 0)
+			{
+				MLFileInfo* FileInfo = nullptr;
+				Result = MLSharedFileGetMLFileInfoByIndex(SharedFileList, 0, &FileInfo);
+				UE_CLOG(Result != MLResult_Ok, LogMagicLeapMediaCodec, Error, TEXT("MLSharedFileGetMLFileInfoByIndex failed with error %s"), UTF8_TO_TCHAR(MLSharedFileGetResultString(Result)));
+				if (Result == MLResult_Ok && FileInfo != nullptr)
+				{
+					MLFileDescriptor FileHandle = -1;
+					Result = MLFileInfoGetFD(FileInfo, &FileHandle);
+					UE_CLOG(Result != MLResult_Ok, LogMagicLeapMediaCodec, Error, TEXT("MLFileInfoGetFD failed with error %s"), UTF8_TO_TCHAR(MLGetResultString(Result)));
+					if (Result == MLResult_Ok && FileHandle > -1)
+					{
+						struct stat NativeFileInfo;
+						fstat(FileHandle, &NativeFileInfo);
+						Result = MLMediaExtractorSetDataSourceForFD(MediaExtractorHandle, FileHandle, 0, NativeFileInfo.st_size);
+						UE_CLOG(Result != MLResult_Ok, LogMagicLeapMediaCodec, Error, TEXT("MLMediaPlayerSetDataSourceForFD failed with error %s"), UTF8_TO_TCHAR(MLGetResultString(Result)));
+					}
+				}
+			}
+
+			MLSharedFileListRelease(&SharedFileList);
+		}
+
+		if (Result != MLResult_Ok)
+		{
+			UE_LOG(LogMagicLeapMediaCodec, Error, TEXT("MLMediaExtractorSetDataSourceForFD for shared file %s failed."), *Url);
+			EventSink.ReceiveMediaEvent(EMediaEvent::MediaOpenFailed);
+			return false;
+		}
+	}
 	else
 	{
 		// open remote media
@@ -790,7 +759,7 @@ void FMagicLeapMediaCodecPlayer::TickFetch(FTimespan DeltaTime, FTimespan Timeco
 						}
 
 						TextureDataPtr->VideoTexturePool.Add((uint64)NativeBuffer, NewMediaTexture);
-					
+
 						if (TextureDataPtr->VideoTexture == nullptr)
 						{
 							FRHIResourceCreateInfo CreateInfo;
@@ -871,57 +840,26 @@ void FMagicLeapMediaCodecPlayer::TickFetch(FTimespan DeltaTime, FTimespan Timeco
 						return;
 					}
 
-					// Clear gl errors as they can creep in from the UE4 renderer.
-					glGetError();
-
-					if (!TextureDataPtr->bContextCreated)
-					{
-						TextureDataPtr->InitContext();
-						TextureDataPtr->bContextCreated = true;
-					}
-					TextureDataPtr->SaveContext();
-					TextureDataPtr->MakeCurrent();
-
 					int32 TextureID = *reinterpret_cast<int32*>(MediaVideoTexture->GetNativeResource());
-					if (TextureDataPtr->Image != EGL_NO_IMAGE_KHR)
-					{
-						eglDestroyImageKHR(eglGetCurrentDisplay(), TextureDataPtr->Image);
-						TextureDataPtr->Image = EGL_NO_IMAGE_KHR;
-					}
+					TextureDataPtr->TextureRenderer.DestroyImageKHR();
 					if (TextureDataPtr->PreviousNativeBuffer != 0 && MLHandleIsValid(TextureDataPtr->PreviousNativeBuffer))
 					{
 						Params.MediaPlayer->ReleaseNativeBuffer_RenderThread(Params.VideoCodecHandle, TextureDataPtr->PreviousNativeBuffer);
 					}
 					TextureDataPtr->PreviousNativeBuffer = nativeBuffer;
 
-					// Wrap latest decoded frame into a new gl texture oject
-					TextureDataPtr->Image = eglCreateImageKHR(TextureDataPtr->Display, EGL_NO_CONTEXT, EGL_NATIVE_BUFFER_ANDROID, (EGLClientBuffer)(void *)nativeBuffer, NULL);
-					if (TextureDataPtr->Image == EGL_NO_IMAGE_KHR)
+					if (TextureDataPtr->TextureRenderer.CreateImageKHR(reinterpret_cast<void*>(nativeBuffer)))
 					{
-						EGLint errorcode = eglGetError();
-						UE_LOG(LogMagicLeapMediaCodec, Error, TEXT("Failed to create EGLImage from the buffer. %d"), errorcode);
-						TextureDataPtr->RestoreContext();
-						return;
-					}
-					glActiveTexture(GL_TEXTURE0);
-					glBindTexture(GL_TEXTURE_EXTERNAL_OES, TextureID);
-					glEGLImageTargetTexture2DOES(GL_TEXTURE_EXTERNAL_OES, TextureDataPtr->Image);
-					glTexParameteri(GL_TEXTURE_EXTERNAL_OES, GL_TEXTURE_MIN_FILTER, GL_LINEAR);
-					glTexParameteri(GL_TEXTURE_EXTERNAL_OES, GL_TEXTURE_MAG_FILTER, GL_LINEAR);
-					glTexParameteri(GL_TEXTURE_EXTERNAL_OES, GL_TEXTURE_WRAP_S, GL_CLAMP_TO_EDGE);
-					glTexParameteri(GL_TEXTURE_EXTERNAL_OES, GL_TEXTURE_WRAP_T, GL_CLAMP_TO_EDGE);
-					glBindTexture(GL_TEXTURE_EXTERNAL_OES, 0);
-
-					TextureDataPtr->RestoreContext();
-
-					TextureDataPtr->bIsVideoTextureValid = Params.MediaPlayer->UpdateTransformMatrix_RenderThread(Params.VideoCodecHandle);
+						TextureDataPtr->TextureRenderer.BindImageKHRToTexture(TextureID);
+						TextureDataPtr->bIsVideoTextureValid = Params.MediaPlayer->UpdateTransformMatrix_RenderThread(Params.VideoCodecHandle);
 
-					if (!TextureDataPtr->bIsVideoTextureValid)
-					{
-						FSamplerStateInitializerRHI SamplerStateInitializer(SF_Bilinear, AM_Clamp, AM_Clamp, AM_Clamp);
-						FSamplerStateRHIRef SamplerStateRHI = RHICreateSamplerState(SamplerStateInitializer);
-						Params.MediaPlayer->RegisterExternalTexture_RenderThread(Params.PlayerGuid, MediaVideoTexture, SamplerStateRHI);
-						TextureDataPtr->bIsVideoTextureValid = true;
+						if (!TextureDataPtr->bIsVideoTextureValid)
+						{
+							FSamplerStateInitializerRHI SamplerStateInitializer(SF_Bilinear, AM_Clamp, AM_Clamp, AM_Clamp);
+							FSamplerStateRHIRef SamplerStateRHI = RHICreateSamplerState(SamplerStateInitializer);
+							Params.MediaPlayer->RegisterExternalTexture_RenderThread(Params.PlayerGuid, MediaVideoTexture, SamplerStateRHI);
+							TextureDataPtr->bIsVideoTextureValid = true;
+						}
 					}
 				});
 		}
@@ -1103,19 +1041,13 @@ bool FMagicLeapMediaCodecPlayer::SelectTrack(EMediaTrackType TrackType, int32 Tr
 {
 	if (TrackInfo.Contains(TrackType) && CurrentState != EMediaState::Preparing)
 	{
-		if (TrackInfo[TrackType].IsValidIndex(TrackIndex))
+		if (TrackInfo[TrackType].IsValidIndex(TrackIndex) && SelectedTrack[TrackType] != TrackIndex)
 		{
-			// TODO: codec needs to be changed.
-			MLResult Result = MLMediaExtractorSelectTrack(MediaExtractorHandle, static_cast<SIZE_T>(TrackInfo[TrackType][TrackIndex].TrackIndex));
-			if (Result == MLResult_Ok)
-			{
-				SelectedTrack[TrackType] = TrackIndex;
-				return true;
-			}
-			else
-			{
-				UE_LOG(LogMagicLeapMediaCodec, Error, TEXT("MLMediaExtractorSelectTrack failed with error %s"), UTF8_TO_TCHAR(MLGetResultString(Result)));
-			}
+			FScopeLock Lock(&CriticalSection);
+			int32 GlobalTrackIndex = TrackInfo[TrackType][TrackIndex].TrackIndex;
+			TrackSelectionQueue.Add(GlobalTrackIndex, { TrackType, TrackIndex });
+			InputWorker.SelectTrack(GlobalTrackIndex, GetTime());
+			return true;
 		}
 	}
 	return false;
@@ -1158,6 +1090,8 @@ void FMagicLeapMediaCodecPlayer::FlushCodecs()
 		}
 	}
 
+	bJustSeeked = true;
+
 	Samples->FlushSamples();
 }
 
@@ -1194,6 +1128,20 @@ void FMagicLeapMediaCodecPlayer::QueueVideoCodecStartTimeReset()
 	bResetVideoCodecStartTime = true;
 }
 
+void FMagicLeapMediaCodecPlayer::SetSelectedTrack(int32 TrackIndex)
+{
+	FScopeLock Lock(&CriticalSection);
+	FTrackSelectionData* SelectedTrackData = TrackSelectionQueue.Find(TrackIndex);
+	if (SelectedTrackData == nullptr)
+	{
+		UE_LOG(LogMagicLeapMediaCodec, Error, TEXT("This is embarrasing. Selected track index %d does not exist in the TrackSelectionQueue."), TrackIndex);
+	}
+	else
+	{
+		SelectedTrack[SelectedTrackData->TrackType] = SelectedTrackData->LocalTrackIndex;
+	}
+}
+
 int64_t FMagicLeapMediaCodecPlayer::MediaDataSourceReadAt(MLHandle media_data_source, size_t position, size_t size, uint8_t *buffer)
 {
 	if (media_data_source == MediaDataSourceHandle && DataSourceArchive.IsValid() && buffer != nullptr)
@@ -1566,11 +1514,24 @@ bool FMagicLeapMediaCodecPlayer::ProcessVideoOutputSample_RenderThread(MLHandle&
 				int32 SelectedAudioTrack = SelectedTrack[EMediaTrackType::Audio];
 				if (SelectedAudioTrack != INDEX_NONE)
 				{
+					// Even when we are seeking, LastAudioPresentationTime is continously set to the old timestamp in ProcessAudioOutputSample().
+					// Because of the large delta between the new seek time and the old audio time, the condition below for
+					// "waiting on audio to catch up" fails and thus CurrentPlaybackTime is never updated.
+					// Media framework in the engine uses CurrentPlaybackTime to dequeue audio samples and send them to the mixer.
+					// Since CurrentPlaybackTime never udpates, LastAudioPresentationTime isnt updated either.
+					// This causes a deadlock in the playback.
+					// To fix this, we ignore the "waiting on audio to catch up" condition when we have "just seeked".
+					// This advances the CurrentPlaybackTime, which then advances LastAudioPresentationTime, and the playback continues normally.
+					if (bJustSeeked)
+					{
+						break;
+					}
+
 					const FTrackData& CurrentAudioTrack = TrackInfo[EMediaTrackType::Audio][SelectedAudioTrack];
 					const FTimespan LastAudioPresentationTime = FTimespan(FPlatformAtomics::AtomicRead(reinterpret_cast<const int64*>(&CurrentAudioTrack.LastPresentationTime)));
 					if ((CurrentPresentationTime - LastAudioPresentationTime) > FTimespan::FromMicroseconds(250))
 					{
-						// UE_LOG(LogMagicLeapMediaCodec, Warning, TEXT("Waiting on audio to catch up."));
+						// UE_LOG(LogMagicLeapMediaCodec, Warning, TEXT("Waiting on audio to catch up. CurrentPresentationTime = %f, LastAudioPresentationTime = %f"), CurrentPresentationTime.GetTotalMilliseconds(), LastAudioPresentationTime.GetTotalMilliseconds());
 						CurrentTrackData.bBufferPendingRender = true;
 						// Reset StartPresentationTime to prevent video stuttering.
 						CurrentTrackData.StartPresentationTime = Timecode - CurrentPresentationTime;
@@ -1620,6 +1581,8 @@ bool FMagicLeapMediaCodecPlayer::ProcessVideoOutputSample_RenderThread(MLHandle&
 		CurrentTrackData.LastPresentationTime = Timecode;
 		CurrentPlaybackTime = CurrentPresentationTime;
 
+		bJustSeeked = false;
+
 		return true;
 	}
 	else if (CurrentTrackData.CurrentBufferIndex == MLMediaCodec_FormatChanged)
diff --git a/Engine/Plugins/Lumin/MagicLeapMedia/Source/MagicLeapMediaCodec/Private/MagicLeapMediaCodecPlayer.h b/Engine/Plugins/Lumin/MagicLeapMedia/Source/MagicLeapMediaCodec/Private/MagicLeapMediaCodecPlayer.h
--- a/Engine/Plugins/Lumin/MagicLeapMedia/Source/MagicLeapMediaCodec/Private/MagicLeapMediaCodecPlayer.h
+++ b/Engine/Plugins/Lumin/MagicLeapMedia/Source/MagicLeapMediaCodec/Private/MagicLeapMediaCodecPlayer.h
@@ -23,25 +23,7 @@
 #include "Serialization/Archive.h"
 #include "Serialization/ArrayReader.h"
 #include "MediaCodecInputWorker.h"
-
-#ifndef EGL_EGLEXT_PROTOTYPES
-#define EGL_EGLEXT_PROTOTYPES
-#endif // !EGL_EGLEXT_PROTOTYPES
-
-#include <EGL/egl.h>
-#include <EGL/eglext.h>
-
-#ifndef GL_GLEXT_PROTOTYPES
-#define GL_GLEXT_PROTOTYPES
-#endif // !GL_GLEXT_PROTOTYPES
-
-#include <GLES2/gl2.h>
-#include <GLES2/gl2ext.h>
-
-#include "ml_api.h"
-#include "ml_media_codec.h"
-#include "ml_audio.h"
-#include "ml_media_data_source.h"
+#include "Lumin/CAPIShims/LuminAPIMediaCodec.h"
 
 class FMediaSamples;
 class IMediaEventSink;
@@ -121,6 +103,7 @@ public:
 	bool IsPlaybackCompleted() const;
 	void SetPlaybackCompleted(bool PlaybackCompleted);
 	void QueueVideoCodecStartTimeReset();
+	void SetSelectedTrack(int32 TrackIndex);
 
 	int64_t MediaDataSourceReadAt(MLHandle media_data_source, size_t position, size_t size, uint8_t *buffer);
 	int64_t MediaDataSourceGetSize(MLHandle media_data_source);
@@ -137,7 +120,7 @@ private:
 	// TODO: Move or rename, get rid of height/width since this is used for both audio and video
 	// Perhaps an FAudioTrackData and FVideoTrackData inheriting from FTrackData?
 	// Also make the codec handle a part of this struct. Will make it easier to perform all operations.
-	struct FTrackData 
+	struct FTrackData
 	{
 		FTrackData(FString Mime, int32 Index) : MimeType(Mime), TrackIndex(Index) {}
 		FString MimeType;
@@ -230,6 +213,7 @@ protected:
 	TAtomic<FTimespan> CurrentPlaybackTime;
 	TAtomic<FTimespan> LastAudioRenderedSampleTime;
 	EMediaSourceType MediaSourceType;
+	FThreadSafeBool bJustSeeked;
 
 	bool bLoopPlayback;
 	bool bPlaybackCompleted;
@@ -242,10 +226,18 @@ protected:
 
 	TQueue<EMediaEvent, EQueueMode::Spsc> MediaEventQueue;
 
-	FMediaCodecInputWorker InputWorker;
+	FMagicLeapMediaCodecInputWorker InputWorker;
 
 	TSharedPtr<FArchive, ESPMode::ThreadSafe> DataSourceArchive;
 
+	struct FTrackSelectionData
+	{
+		EMediaTrackType TrackType;
+		int32 LocalTrackIndex;
+	};
+
+	TMap<int32, FTrackSelectionData> TrackSelectionQueue;
+
 private:
 	// Non-Interface functions
 	virtual void GetTrackInformation();
diff --git a/Engine/Plugins/Lumin/MagicLeapMedia/Source/MagicLeapMediaCodec/Private/MediaCodecInputWorker.cpp b/Engine/Plugins/Lumin/MagicLeapMedia/Source/MagicLeapMediaCodec/Private/MediaCodecInputWorker.cpp
--- a/Engine/Plugins/Lumin/MagicLeapMedia/Source/MagicLeapMediaCodec/Private/MediaCodecInputWorker.cpp
+++ b/Engine/Plugins/Lumin/MagicLeapMedia/Source/MagicLeapMediaCodec/Private/MediaCodecInputWorker.cpp
@@ -7,12 +7,11 @@
 #include "MagicLeapMediaCodecPlayer.h"
 #include "Lumin/LuminPlatformAffinity.h"
 #include "IMagicLeapMediaCodecModule.h"
+#include "Lumin/CAPIShims/LuminAPIMediaExtractor.h"
+#include "Lumin/CAPIShims/LuminAPIMediaError.h"
+#include "Lumin/CAPIShims/LuminAPIMediaFormat.h"
 
-#include "ml_media_extractor.h"
-#include "ml_media_codec.h"
-#include "ml_media_error.h"
-
-FMediaCodecInputWorker::FMediaCodecInputWorker()
+FMagicLeapMediaCodecInputWorker::FMagicLeapMediaCodecInputWorker()
 : OwnerPlayer(nullptr)
 , ExtractorHandle(nullptr)
 , CriticalSection(nullptr)
@@ -24,12 +23,12 @@ FMediaCodecInputWorker::FMediaCodecInputWorker()
 , bReachedInputEndOfStream(false)
 {}
 
-FMediaCodecInputWorker::~FMediaCodecInputWorker()
+FMagicLeapMediaCodecInputWorker::~FMagicLeapMediaCodecInputWorker()
 {
 	DestroyThread();
 }
 
-void FMediaCodecInputWorker::InitThread(FMagicLeapMediaCodecPlayer& InOwnerPlayer, MLHandle& InExtractorHandle, FCriticalSection& InCriticalSection, FCriticalSection& InGT_IT_Mutex, FCriticalSection& InRT_IT_Mutex)
+void FMagicLeapMediaCodecInputWorker::InitThread(FMagicLeapMediaCodecPlayer& InOwnerPlayer, MLHandle& InExtractorHandle, FCriticalSection& InCriticalSection, FCriticalSection& InGT_IT_Mutex, FCriticalSection& InRT_IT_Mutex)
 {
 	OwnerPlayer = &InOwnerPlayer;
 	ExtractorHandle = &InExtractorHandle;
@@ -46,7 +45,7 @@ void FMediaCodecInputWorker::InitThread(FMagicLeapMediaCodecPlayer& InOwnerPlaye
 	}
 }
 
-void FMediaCodecInputWorker::DestroyThread()
+void FMagicLeapMediaCodecInputWorker::DestroyThread()
 {
 	StopTaskCounter.Increment();
 	if (Semaphore != nullptr)
@@ -56,13 +55,13 @@ void FMediaCodecInputWorker::DestroyThread()
 		FGenericPlatformProcess::ReturnSynchEventToPool(Semaphore);
 		Semaphore = nullptr;
 		delete Thread;
-		Thread = nullptr;		
+		Thread = nullptr;
 	}
 }
 
-uint32 FMediaCodecInputWorker::Run()
+uint32 FMagicLeapMediaCodecInputWorker::Run()
 {
-	FInputWorkerTask QueuedTask;
+	FMagicLeapMediaCodecInputWorkerTask QueuedTask;
 
 	while (StopTaskCounter.GetValue() == 0)
 	{
@@ -91,12 +90,18 @@ uint32 FMediaCodecInputWorker::Run()
 
 		if (IncomingTasks.Dequeue(QueuedTask))
 		{
-			if (QueuedTask.TaskType == EInputWorkerTaskType::Seek)
+			if (QueuedTask.TaskType == EMagicLeapMediaCodecInputWorkerTaskType::Seek)
 			{
 				FScopeLock LockGT(GT_IT_Mutex);
 				FScopeLock LockRT(RT_IT_Mutex);
 				Seek_WorkerThread(QueuedTask.SeekTime);
 			}
+			else if (QueuedTask.TaskType == EMagicLeapMediaCodecInputWorkerTaskType::SelectTrack)
+			{
+				FScopeLock LockGT(GT_IT_Mutex);
+				FScopeLock LockRT(RT_IT_Mutex);
+				SelectTrack_WorkerThread(QueuedTask.TrackIndex, QueuedTask.SeekTime);
+			}
 		}
 
 		ProcessInputSample_WorkerThread();
@@ -110,7 +115,7 @@ uint32 FMediaCodecInputWorker::Run()
 	return 0;
 }
 
-void FMediaCodecInputWorker::WakeUp()
+void FMagicLeapMediaCodecInputWorker::WakeUp()
 {
 	if (Semaphore != nullptr)
 	{
@@ -118,17 +123,22 @@ void FMediaCodecInputWorker::WakeUp()
 	}
 }
 
-void FMediaCodecInputWorker::Seek(FTimespan SeekTime)
+void FMagicLeapMediaCodecInputWorker::Seek(FTimespan SeekTime)
 {
-	IncomingTasks.Enqueue(FInputWorkerTask(EInputWorkerTaskType::Seek, SeekTime));
+	IncomingTasks.Enqueue(FMagicLeapMediaCodecInputWorkerTask(EMagicLeapMediaCodecInputWorkerTaskType::Seek, SeekTime));
 }
 
-bool FMediaCodecInputWorker::HasReachedInputEOS() const
+void FMagicLeapMediaCodecInputWorker::SelectTrack(int32 TrackIndex, const FTimespan& SeekTime)
+{
+	IncomingTasks.Enqueue(FMagicLeapMediaCodecInputWorkerTask(EMagicLeapMediaCodecInputWorkerTaskType::SelectTrack, TrackIndex, SeekTime));
+}
+
+bool FMagicLeapMediaCodecInputWorker::HasReachedInputEOS() const
 {
 	return bReachedInputEndOfStream;
 }
 
-void FMediaCodecInputWorker::ProcessInputSample_WorkerThread()
+void FMagicLeapMediaCodecInputWorker::ProcessInputSample_WorkerThread()
 {
 	{
 		FScopeLock Lock(CriticalSection);
@@ -164,7 +174,7 @@ void FMediaCodecInputWorker::ProcessInputSample_WorkerThread()
 		else
 		{
 			UE_LOG(LogMagicLeapMediaCodec, Display, TEXT("negative track index from MLMediaExtractorGetSampleTrackIndex. Reached input EOS."));
-		
+
 			{
 				FScopeLock Lock(CriticalSection);
 				bReachedInputEndOfStream = true;
@@ -229,10 +239,9 @@ void FMediaCodecInputWorker::ProcessInputSample_WorkerThread()
 	UE_CLOG(Result != MLResult_Ok, LogMagicLeapMediaCodec, Error, TEXT("MLMediaExtractorAdvance() failed with error %s"), UTF8_TO_TCHAR(MLMediaResultGetString(Result)));
 }
 
-bool FMediaCodecInputWorker::Seek_WorkerThread(const FTimespan& SeekTime)
+bool FMagicLeapMediaCodecInputWorker::Seek_WorkerThread(const FTimespan& SeekTime)
 {
-	UE_LOG(LogMagicLeapMediaCodec, Display, TEXT("Seek(%f)"), SeekTime.GetTotalMilliseconds());
-	MLResult Result = MLMediaExtractorSeekTo(*ExtractorHandle, SeekTime.GetTotalMicroseconds(), MLMediaSeekMode_Closest_Sync);
+	MLResult Result = MLMediaExtractorSeekTo(*ExtractorHandle, static_cast<int64>(SeekTime.GetTotalMicroseconds()), MLMediaSeekMode_Closest_Sync);
 	if (Result == MLResult_Ok)
 	{
 		// TODO: Make thread safe.
@@ -249,3 +258,22 @@ bool FMediaCodecInputWorker::Seek_WorkerThread(const FTimespan& SeekTime)
 		return false;
 	}
 }
+
+void FMagicLeapMediaCodecInputWorker::SelectTrack_WorkerThread(int32 TrackIndex, const FTimespan& SeekTime)
+{
+	MLHandle TrackFormatHandle = ML_INVALID_HANDLE;
+	MLResult Result = MLMediaExtractorGetTrackFormat(*ExtractorHandle, TrackIndex, &TrackFormatHandle);
+	UE_CLOG(Result != MLResult_Ok, LogMagicLeapMediaCodec, Error, TEXT("MLMediaExtractorGetTrackFormat() failed with error %s"), UTF8_TO_TCHAR(MLMediaResultGetString(Result)));
+
+	Result = MLMediaExtractorSelectTrack(*ExtractorHandle, static_cast<SIZE_T>(TrackIndex));
+	if (Result == MLResult_Ok)
+	{
+		OwnerPlayer->SetSelectedTrack(TrackIndex);
+		Seek(SeekTime);
+	}
+
+	UE_CLOG(Result != MLResult_Ok, LogMagicLeapMediaCodec, Error, TEXT("MLMediaExtractorSelectTrack failed with error %s"), UTF8_TO_TCHAR(MLGetResultString(Result)));
+
+	Result = MLMediaFormatDestroy(TrackFormatHandle);
+	UE_CLOG(Result != MLResult_Ok, LogMagicLeapMediaCodec, Error, TEXT("MLMediaFormatDestroy() failed with error %s"), UTF8_TO_TCHAR(MLMediaResultGetString(Result)));
+}
diff --git a/Engine/Plugins/Lumin/MagicLeapMedia/Source/MagicLeapMediaCodec/Private/MediaCodecInputWorker.h b/Engine/Plugins/Lumin/MagicLeapMedia/Source/MagicLeapMediaCodec/Private/MediaCodecInputWorker.h
--- a/Engine/Plugins/Lumin/MagicLeapMedia/Source/MagicLeapMediaCodec/Private/MediaCodecInputWorker.h
+++ b/Engine/Plugins/Lumin/MagicLeapMedia/Source/MagicLeapMediaCodec/Private/MediaCodecInputWorker.h
@@ -6,40 +6,49 @@
 #include "HAL/Runnable.h"
 #include "HAL/ThreadSafeCounter.h"
 #include "Containers/Queue.h"
-
-#include "ml_api.h"
+#include "Lumin/CAPIShims/LuminAPI.h"
 
 class FMagicLeapMediaCodecPlayer;
 
-enum class EInputWorkerTaskType : uint8
+enum class EMagicLeapMediaCodecInputWorkerTaskType : uint8
 {
 	None,
-	Seek
+	Seek,
+	SelectTrack
 };
 
-struct FInputWorkerTask
+struct FMagicLeapMediaCodecInputWorkerTask
 {
 public:
-	EInputWorkerTaskType TaskType;
+	EMagicLeapMediaCodecInputWorkerTaskType TaskType;
 	FTimespan SeekTime;
+	int32 TrackIndex;
 
 public:
-	FInputWorkerTask()
-	: TaskType(EInputWorkerTaskType::None)
+	FMagicLeapMediaCodecInputWorkerTask()
+	: TaskType(EMagicLeapMediaCodecInputWorkerTaskType::None)
 	, SeekTime(FTimespan::Zero())
+	, TrackIndex(0)
+	{}
+
+	FMagicLeapMediaCodecInputWorkerTask(EMagicLeapMediaCodecInputWorkerTaskType InTaskType, const FTimespan& InSeekTime)
+	: TaskType(InTaskType)
+	, SeekTime(InSeekTime)
+	, TrackIndex(0)
 	{}
 
-	FInputWorkerTask(EInputWorkerTaskType InTaskType, const FTimespan& InSeekTime)
+	FMagicLeapMediaCodecInputWorkerTask(EMagicLeapMediaCodecInputWorkerTaskType InTaskType, int32 InTrackIndex, const FTimespan& InSeekTime)
 	: TaskType(InTaskType)
 	, SeekTime(InSeekTime)
+	, TrackIndex(InTrackIndex)
 	{}
 };
 
-class FMediaCodecInputWorker : public FRunnable
+class FMagicLeapMediaCodecInputWorker : public FRunnable
 {
 public:
-	FMediaCodecInputWorker();
-	virtual ~FMediaCodecInputWorker();
+	FMagicLeapMediaCodecInputWorker();
+	virtual ~FMagicLeapMediaCodecInputWorker();
 
 	void InitThread(FMagicLeapMediaCodecPlayer& InOwnerPlayer, MLHandle& InExtractorHandle, FCriticalSection& InCriticalSection, FCriticalSection& InGT_IT_Mutex, FCriticalSection& InRT_IT_Mutex);
 	void DestroyThread();
@@ -47,12 +56,14 @@ public:
 
 	void WakeUp();
 	void Seek(FTimespan SeekTime);
+	void SelectTrack(int32 TrackIndex, const FTimespan& SeekTime);
 
 	bool HasReachedInputEOS() const;
 
 private:
 	void ProcessInputSample_WorkerThread();
 	bool Seek_WorkerThread(const FTimespan& SeekTime);
+	void SelectTrack_WorkerThread(int32 TrackIndex, const FTimespan& SeekTime);
 
 	FMagicLeapMediaCodecPlayer* OwnerPlayer;
 	MLHandle* ExtractorHandle;
@@ -66,5 +77,5 @@ private:
 
 	bool bReachedInputEndOfStream;
 
-	TQueue<FInputWorkerTask, EQueueMode::Spsc> IncomingTasks;
+	TQueue<FMagicLeapMediaCodecInputWorkerTask, EQueueMode::Mpsc> IncomingTasks;
 };
diff --git a/Engine/Plugins/Lumin/MagicLeapMedia/Source/MagicLeapMediaCodecFactory/Private/MagicLeapMediaCodecFactoryModule.cpp b/Engine/Plugins/Lumin/MagicLeapMedia/Source/MagicLeapMediaCodecFactory/Private/MagicLeapMediaCodecFactoryModule.cpp
--- a/Engine/Plugins/Lumin/MagicLeapMedia/Source/MagicLeapMediaCodecFactory/Private/MagicLeapMediaCodecFactoryModule.cpp
+++ b/Engine/Plugins/Lumin/MagicLeapMedia/Source/MagicLeapMediaCodecFactory/Private/MagicLeapMediaCodecFactoryModule.cpp
@@ -114,10 +114,14 @@ public:
 	virtual void StartupModule() override
 	{
 		// supported file extensions
+		SupportedFileExtensions.Add(TEXT("3gp"));
 		SupportedFileExtensions.Add(TEXT("mp4"));
-		SupportedFileExtensions.Add(TEXT("3gpp"));
+		SupportedFileExtensions.Add(TEXT("mp4a"));
 		SupportedFileExtensions.Add(TEXT("aac"));
-		SupportedFileExtensions.Add(TEXT("m3u8"));
+		SupportedFileExtensions.Add(TEXT("ts"));
+		SupportedFileExtensions.Add(TEXT("mkv"));
+		SupportedFileExtensions.Add(TEXT("webm"));
+		SupportedFileExtensions.Add(TEXT("m4v"));
 
     	// supported platforms
     	SupportedPlatforms.Add(TEXT("Lumin"));
@@ -126,6 +130,8 @@ public:
 		SupportedUriSchemes.Add(TEXT("file"));
 		SupportedUriSchemes.Add(TEXT("http"));
 		SupportedUriSchemes.Add(TEXT("https"));
+		// custom schema for user shared files
+		SupportedUriSchemes.Add(TEXT("mlshared"));
 		// Not supporting streaming right now.
 		// SupportedUriSchemes.Add(TEXT("httpd"));
 		// SupportedUriSchemes.Add(TEXT("mms"));
diff --git a/Engine/Plugins/Lumin/MagicLeapMedia/Source/MagicLeapMediaFactory/Private/MagicLeapMediaFactoryModule.cpp b/Engine/Plugins/Lumin/MagicLeapMedia/Source/MagicLeapMediaFactory/Private/MagicLeapMediaFactoryModule.cpp
--- a/Engine/Plugins/Lumin/MagicLeapMedia/Source/MagicLeapMediaFactory/Private/MagicLeapMediaFactoryModule.cpp
+++ b/Engine/Plugins/Lumin/MagicLeapMedia/Source/MagicLeapMediaFactory/Private/MagicLeapMediaFactoryModule.cpp
@@ -5,6 +5,7 @@
 #include "Modules/ModuleInterface.h"
 #include "Modules/ModuleManager.h"
 #include "IMediaModule.h"
+#include "IMediaOptions.h"
 #include "IMagicLeapMediaModule.h"
 #include "MagicLeapMediaFactoryPrivate.h"
 #include "IMediaPlayerFactory.h"
@@ -21,7 +22,7 @@ class FMagicLeapMediaFactoryModule : public IMediaPlayerFactory, public IModuleI
 public:
 
 	/** IMediaPlayerFactory interface */
-	virtual bool CanPlayUrl(const FString& Url, const IMediaOptions* /*Options*/, TArray<FText>* /*OutWarnings*/, TArray<FText>* OutErrors) const override
+	virtual bool CanPlayUrl(const FString& Url, const IMediaOptions* Options, TArray<FText>* OutWarnings, TArray<FText>* OutErrors) const override
 	{
 		FString Scheme;
 		FString Location;
@@ -63,6 +64,15 @@ public:
 			}
 		}
 
+		// check options
+		if ((OutWarnings != nullptr) && (Options != nullptr))
+		{
+			if (Options->GetMediaOption("PrecacheFile", false))
+			{
+				OutWarnings->Add(LOCTEXT("PrecachingNotSupported", "Precaching is not supported in Magic Leap Media. Use Magic Leap Media Codec instead if this feature is necessarry."));
+			}
+		}
+
 		return true;
 	}
 
@@ -102,10 +112,14 @@ public:
 	virtual void StartupModule() override
 	{
 		// supported file extensions
+		SupportedFileExtensions.Add(TEXT("3gp"));
 		SupportedFileExtensions.Add(TEXT("mp4"));
-		SupportedFileExtensions.Add(TEXT("3gpp"));
+		SupportedFileExtensions.Add(TEXT("mp4a"));
 		SupportedFileExtensions.Add(TEXT("aac"));
-		SupportedFileExtensions.Add(TEXT("m3u8"));
+		SupportedFileExtensions.Add(TEXT("ts"));
+		SupportedFileExtensions.Add(TEXT("mkv"));
+		SupportedFileExtensions.Add(TEXT("webm"));
+		SupportedFileExtensions.Add(TEXT("m4v"));
 
     	// supported platforms
     	SupportedPlatforms.Add(TEXT("Lumin"));
@@ -115,6 +129,8 @@ public:
 		SupportedUriSchemes.Add(TEXT("http"));
 		SupportedUriSchemes.Add(TEXT("https"));
 		SupportedUriSchemes.Add(TEXT("rtsp"));
+		// custom schema for user shared files
+		SupportedUriSchemes.Add(TEXT("mlshared"));
 		// Not supporting streaming right now.
 		// SupportedUriSchemes.Add(TEXT("httpd"));
 		// SupportedUriSchemes.Add(TEXT("mms"));
